//UE4/Main: 'Build Engine Localization' - LogGatherTextFromSourceCommandlet and LogLocTextHelper warnings
#fix Duplicate LOCTEXT label entry

diff --git a/Engine/Source/Editor/AnimationModifiers/Private/SAnimationModifiersTab.cpp b/Engine/Source/Editor/AnimationModifiers/Private/SAnimationModifiersTab.cpp
--- a/Engine/Source/Editor/AnimationModifiers/Private/SAnimationModifiersTab.cpp
+++ b/Engine/Source/Editor/AnimationModifiers/Private/SAnimationModifiersTab.cpp
@@ -301,7 +301,7 @@ void SAnimationModifiersTab::CreateInstanceDetailsView()
 
 FReply SAnimationModifiersTab::OnApplyAllModifiersClicked()
 {
-	FScopedTransaction Transaction(LOCTEXT("ApplyModifiersTransaction", "Applying All Animation Modifier(s)"));
+	FScopedTransaction Transaction(LOCTEXT("ApplyAllModifiersTransaction", "Applying All Animation Modifier(s)"));
 	const TArray<UAnimationModifier*>& ModifierInstances = AssetUserData->GetAnimationModifierInstances();	
 	ApplyModifiers(ModifierInstances);
 	return FReply::Handled();
