Another LOCTEXT duplicate label fix

diff --git a/Engine/Source/Editor/AnimationModifiers/Private/SAnimationModifiersTab.cpp b/Engine/Source/Editor/AnimationModifiers/Private/SAnimationModifiersTab.cpp
--- a/Engine/Source/Editor/AnimationModifiers/Private/SAnimationModifiersTab.cpp
+++ b/Engine/Source/Editor/AnimationModifiers/Private/SAnimationModifiersTab.cpp
@@ -499,7 +499,7 @@ void SAnimationModifiersTab::RevertModifiers(const TArray<UAnimationModifier*>&
 	else if (Skeleton != nullptr)
 	{
 		// Double check with the user for reverting all modifiers from referenced animation sequences for the skeleton
-		bRevert = OpenMsgDlgInt(EAppMsgType::YesNo, LOCTEXT("ApplyingSkeletonModifierPopupText", "Are you sure you want to revert the modifiers from all animation sequences referenced by the current skeleton?"), FText::FromString("Are you sure?")) == EAppReturnType::Yes;
+		bRevert = OpenMsgDlgInt(EAppMsgType::YesNo, LOCTEXT("RevertingSkeletonModifierPopupText", "Are you sure you want to revert the modifiers from all animation sequences referenced by the current skeleton?"), FText::FromString("Are you sure?")) == EAppReturnType::Yes;
 
 		if ( bRevert)
 		{
