Crash when applying / reverting Animation Modifiers
#fix Add Begin/EndOperation calls which are now required

diff --git a/Engine/Source/Editor/AnimationModifiers/Private/AnimationModifier.cpp b/Engine/Source/Editor/AnimationModifiers/Private/AnimationModifier.cpp
--- a/Engine/Source/Editor/AnimationModifiers/Private/AnimationModifier.cpp
+++ b/Engine/Source/Editor/AnimationModifiers/Private/AnimationModifier.cpp
@@ -39,7 +39,9 @@ void UAnimationModifier::ApplyToAnimationSequence(class UAnimSequence* InAnimati
 	OnApply(CurrentAnimSequence);
 
 	// Apply transaction
+	ModifierTransaction.BeginOperation();
 	ModifierTransaction.Apply();
+	ModifierTransaction.EndOperation();
 
 	GLog->RemoveOutputDevice(&OutputLog);
 
@@ -59,7 +61,9 @@ void UAnimationModifier::ApplyToAnimationSequence(class UAnimSequence* InAnimati
 	// Revert changes if necessary, otherwise post edit and refresh animation data
 	if (bShouldRevert)
 	{
+		AnimationDataTransaction.BeginOperation();
 		AnimationDataTransaction.Apply();
+		AnimationDataTransaction.EndOperation();
 		CurrentAnimSequence->RefreshCacheData();
 		CurrentAnimSequence->RefreshCurveData();
 	}
@@ -112,7 +116,9 @@ void UAnimationModifier::RevertFromAnimationSequence(class UAnimSequence* InAnim
 	OnRevert(CurrentAnimSequence);
 
 	// Apply transaction
+	Transaction.BeginOperation();
 	Transaction.Apply();
+	Transaction.EndOperation();
 
 	UpdateCompressedAnimationData();
 
