Running Anim Modifiers on a selection of animations

diff --git a/Engine/Source/Developer/AssetTools/AssetTools.Build.cs b/Engine/Source/Developer/AssetTools/AssetTools.Build.cs
--- a/Engine/Source/Developer/AssetTools/AssetTools.Build.cs
+++ b/Engine/Source/Developer/AssetTools/AssetTools.Build.cs
@@ -59,6 +59,7 @@ public class AssetTools : ModuleRules
                 "SkeletalMeshEditor",
                 "AnimationEditor",
                 "AnimationBlueprintEditor",
+                "AnimationModifiers"
             }
 		);
 
@@ -83,6 +84,7 @@ public class AssetTools : ModuleRules
                 "SkeletalMeshEditor",
                 "AnimationEditor",
                 "AnimationBlueprintEditor",
+                "AnimationModifiers"
             }
 		);
 	}
diff --git a/Engine/Source/Developer/AssetTools/Private/AssetTypeActions/AssetTypeActions_AnimSequence.cpp b/Engine/Source/Developer/AssetTools/Private/AssetTypeActions/AssetTypeActions_AnimSequence.cpp
--- a/Engine/Source/Developer/AssetTools/Private/AssetTypeActions/AssetTypeActions_AnimSequence.cpp
+++ b/Engine/Source/Developer/AssetTools/Private/AssetTypeActions/AssetTypeActions_AnimSequence.cpp
@@ -15,7 +15,8 @@
 #include "AssetTools.h"
 #include "IContentBrowserSingleton.h"
 #include "ContentBrowserModule.h"
-
+#include "IAnimationModifiersModule.h"
+#include "Algo/Transform.h"
 
 #define LOCTEXT_NAMESPACE "AssetTypeActions"
 
@@ -44,6 +45,13 @@ void FAssetTypeActions_AnimSequence::GetActions( const TArray<UObject*>& InObjec
 		FUIAction(FExecuteAction::CreateSP(this, &FAssetTypeActions_AnimSequence::ExecuteReimportWithNewSource, Sequences))
 		);
 
+	MenuBuilder.AddMenuEntry(
+		LOCTEXT("AnimSequence_AddAnimationModifier", "Add Animation Modifier(s)"),
+		LOCTEXT("AnimSequence_AddAnimationModifierTooltip", "Apply new animation modifier(s)."),
+		FSlateIcon(FEditorStyle::GetStyleSetName(), "ClassIcon.AnimationModifier"),
+		FUIAction(FExecuteAction::CreateSP(this, &FAssetTypeActions_AnimSequence::ExecuteAddNewAnimationModifier, Sequences))
+	);
+
 	FAssetTypeActions_AnimationAsset::GetActions(InObjects, MenuBuilder);
 }
 
@@ -137,6 +145,26 @@ void FAssetTypeActions_AnimSequence::ExecuteNewPoseAsset(TArray<TWeakObjectPtr<U
 	CreateAnimationAssets(Objects, UPoseAsset::StaticClass(), Factory, DefaultSuffix, FOnConfigureFactory::CreateSP(this, &FAssetTypeActions_AnimSequence::ConfigureFactoryForPoseAsset));
 }
 
+void FAssetTypeActions_AnimSequence::ExecuteAddNewAnimationModifier(TArray<TWeakObjectPtr<UAnimSequence>> Objects)
+{
+	TArray<UAnimSequence*> AnimSequences;
+
+	Algo::TransformIf(Objects, AnimSequences, 
+	[](const TWeakObjectPtr<UAnimSequence>& WeakAnimSequence)
+	{
+		return WeakAnimSequence.Get() && WeakAnimSequence->IsA<UAnimSequence>();
+	},
+	[](const TWeakObjectPtr<UAnimSequence>& WeakAnimSequence)
+	{
+		return WeakAnimSequence.Get();
+	});
+
+	if (IAnimationModifiersModule* Module = FModuleManager::Get().LoadModulePtr<IAnimationModifiersModule>("AnimationModifiers"))
+	{
+		Module->ShowAddAnimationModifierWindow(AnimSequences);
+	}
+}
+
 bool FAssetTypeActions_AnimSequence::ConfigureFactoryForAnimComposite(UFactory* AssetFactory, UAnimSequence* SourceAnimation) const
 {
 	UAnimCompositeFactory* CompositeFactory = CastChecked<UAnimCompositeFactory>(AssetFactory);
diff --git a/Engine/Source/Developer/AssetTools/Private/AssetTypeActions/AssetTypeActions_AnimSequence.h b/Engine/Source/Developer/AssetTools/Private/AssetTypeActions/AssetTypeActions_AnimSequence.h
--- a/Engine/Source/Developer/AssetTools/Private/AssetTypeActions/AssetTypeActions_AnimSequence.h
+++ b/Engine/Source/Developer/AssetTools/Private/AssetTypeActions/AssetTypeActions_AnimSequence.h
@@ -38,6 +38,9 @@ private:
 	/** Handler for when Create Pose Asset is selected */
 	void ExecuteNewPoseAsset(TArray<TWeakObjectPtr<UAnimSequence>> Objects) const;
 
+	/** Handler for when Add Animation Modifier is selected */
+	void ExecuteAddNewAnimationModifier(TArray<TWeakObjectPtr<UAnimSequence>> Objects);
+
 	/** Delegate handler passed to CreateAnimationAssets when creating AnimComposites */
 	bool ConfigureFactoryForAnimComposite(UFactory* AssetFactory, UAnimSequence* SourceAnimation) const;
 
diff --git a/Engine/Source/Editor/AnimationModifiers/Private/AnimationModifierHelpers.h b/Engine/Source/Editor/AnimationModifiers/Private/AnimationModifierHelpers.h
new file mode 100644
--- /dev/null
+++ b/Engine/Source/Editor/AnimationModifiers/Private/AnimationModifierHelpers.h
@@ -0,0 +1,61 @@
+// Copyright 1998-2019 Epic Games, Inc. All Rights Reserved.
+
+#include "ClassViewerModule.h"
+#include "Widgets/SWidget.h"
+#include "Widgets/Layout/SBox.h"
+#include "Widgets/SBoxPanel.h"
+#include "Modules/ModuleManager.h"
+#include "Templates/SharedPointer.h"
+#include "ClassViewerFilter.h"
+#include "AnimationModifier.h"
+
+class FAnimationModifierHelpers
+{
+public:
+	/** ClassViewerFilter for Animation Modifier classes */
+	class FModifierClassFilter : public IClassViewerFilter
+	{
+	public:
+		bool IsClassAllowed(const FClassViewerInitializationOptions& InInitOptions, const UClass* InClass, TSharedRef< FClassViewerFilterFuncs > InFilterFuncs) override
+		{
+			return InClass->IsChildOf(UAnimationModifier::StaticClass());
+		}
+
+		virtual bool IsUnloadedClassAllowed(const FClassViewerInitializationOptions& InInitOptions, const TSharedRef< const IUnloadedBlueprintData > InClass, TSharedRef< FClassViewerFilterFuncs > InFilterFuncs) override
+		{
+			return InClass->IsChildOf(UAnimationModifier::StaticClass());
+		}
+	};
+
+	static TSharedRef<SWidget> GetModifierPicker(const FOnClassPicked& OnClassPicked)
+	{
+		FClassViewerInitializationOptions Options;
+		Options.bShowUnloadedBlueprints = true;
+		Options.bShowNoneOption = false;
+		TSharedPtr<FModifierClassFilter> ClassFilter = MakeShareable(new FModifierClassFilter);
+		Options.ClassFilter = ClassFilter;
+
+		return SNew(SBox)
+		.WidthOverride(280)
+		[
+			SNew(SVerticalBox)
+			+ SVerticalBox::Slot()
+			.AutoHeight()
+			.MaxHeight(500)
+			[
+				FModuleManager::LoadModuleChecked<FClassViewerModule>("ClassViewer").CreateClassViewer(Options, OnClassPicked)
+			]
+		];
+	}
+
+	/** Creates a new Modifier instance to store with the current asset */
+	static UAnimationModifier* CreateModifierInstance(UObject* Outer, UClass* InClass, UObject* Template = nullptr)
+	{
+		checkf(Outer, TEXT("Invalid outer value for modifier instantiation"));
+		UAnimationModifier* ProcessorInstance = NewObject<UAnimationModifier>(Outer, InClass, NAME_None, RF_NoFlags, Template);
+		checkf(ProcessorInstance, TEXT("Unable to instantiate modifier class"));
+		ProcessorInstance->SetFlags(RF_Transactional);
+		return ProcessorInstance;
+	}
+
+};
\ No newline at end of file
diff --git a/Engine/Source/Editor/AnimationModifiers/Private/AnimationModifiersModule.cpp b/Engine/Source/Editor/AnimationModifiers/Private/AnimationModifiersModule.cpp
--- a/Engine/Source/Editor/AnimationModifiers/Private/AnimationModifiersModule.cpp
+++ b/Engine/Source/Editor/AnimationModifiers/Private/AnimationModifiersModule.cpp
@@ -12,6 +12,9 @@
 
 #include "Modules/ModuleManager.h"
 #include "PropertyEditorModule.h" 
+#include "SAnimationModifierContentBrowserWindow.h"
+#include "Framework/Application/SlateApplication.h"
+#include "Interfaces/IMainFrameModule.h"
 
 #define LOCTEXT_NAMESPACE "AnimationModifiersModule"
 
@@ -65,4 +68,31 @@ void FAnimationModifiersModule::ShutdownModule()
 	RegisteredApplicationModes.Empty();
 }
 
+void FAnimationModifiersModule::ShowAddAnimationModifierWindow(const TArray<UAnimSequence*>& InSequences)
+{
+	TSharedPtr<SAnimationModifierContentBrowserWindow> WindowContent;
+
+	TSharedRef<SWindow> Window = SNew(SWindow)
+		.Title(LOCTEXT("WindowTitle", "Add Animation Modifier(s)"))
+		.SizingRule(ESizingRule::UserSized)
+		.ClientSize(FVector2D(500, 500));
+
+	Window->SetContent
+	(
+		SAssignNew(WindowContent, SAnimationModifierContentBrowserWindow)
+		.WidgetWindow(Window)
+		.AnimSequences(InSequences)
+	);
+
+	TSharedPtr<SWindow> ParentWindow;
+
+	if (FModuleManager::Get().IsModuleLoaded("MainFrame"))
+	{
+		IMainFrameModule& MainFrame = FModuleManager::LoadModuleChecked<IMainFrameModule>("MainFrame");
+		ParentWindow = MainFrame.GetParentWindow();
+	}
+
+	FSlateApplication::Get().AddModalWindow(Window, ParentWindow, false);
+}
+
 #undef LOCTEXT_NAMESPACE // "AnimationModifiersModule"
diff --git a/Engine/Source/Editor/AnimationModifiers/Private/AnimationModifiersModule.h b/Engine/Source/Editor/AnimationModifiers/Private/AnimationModifiersModule.h
--- a/Engine/Source/Editor/AnimationModifiers/Private/AnimationModifiersModule.h
+++ b/Engine/Source/Editor/AnimationModifiers/Private/AnimationModifiersModule.h
@@ -18,6 +18,10 @@ public:
 	/** Called before the module is unloaded, right before the module object is destroyed */
 	virtual void ShutdownModule() override;
 
+	/** Begin IAnimationModifiersModule overrides */
+	virtual void ShowAddAnimationModifierWindow(const TArray<UAnimSequence*>& InSequences) override;
+	/** End IAnimationModifiersModule overrides */
+
 protected:
 	/** Callback for extending an application mode */
 	TSharedRef<FApplicationMode> ExtendApplicationMode(const FName ModeName, TSharedRef<FApplicationMode> InMode);
diff --git a/Engine/Source/Editor/AnimationModifiers/Private/SAnimationModifierContentBrowserWindow.cpp b/Engine/Source/Editor/AnimationModifiers/Private/SAnimationModifierContentBrowserWindow.cpp
new file mode 100644
--- /dev/null
+++ b/Engine/Source/Editor/AnimationModifiers/Private/SAnimationModifierContentBrowserWindow.cpp
@@ -0,0 +1,267 @@
+// Copyright 1998-2019 Epic Games, Inc. All Rights Reserved.
+
+#include "SAnimationModifierContentBrowserWindow.h"
+#include "AnimationModifierHelpers.h"
+#include "Widgets/SOverlay.h"
+#include "Widgets/SBoxPanel.h"
+#include "Widgets/Layout/SBorder.h"
+#include "Widgets/Layout/SUniformGridPanel.h"
+#include "Widgets/Input/SComboButton.h"
+#include "Widgets/Text/STextBlock.h"
+
+#include "EditorStyleSet.h"
+#include "SModifierListview.h"
+#include "PropertyEditorModule.h"
+
+#include "Animation/AnimSequence.h"
+#include "AnimationModifiersAssetUserData.h"
+#include "Dialogs/Dialogs.h"
+#include "ScopedTransaction.h"
+
+#define LOCTEXT_NAMESPACE "AnimationModifierContentBrowserWindow"
+
+void SAnimationModifierContentBrowserWindow::Construct(const FArguments& InArgs)
+{
+	CreateInstanceDetailsView();
+
+	WidgetWindow = InArgs._WidgetWindow;
+	AnimSequences = InArgs._AnimSequences;
+
+	FOnGetContent GetContent = FOnGetContent::CreateLambda(
+		[this]()
+	{
+		return FAnimationModifierHelpers::GetModifierPicker(FOnClassPicked::CreateRaw(this, &SAnimationModifierContentBrowserWindow::OnModifierPicked));
+	});
+
+	this->ChildSlot
+	[
+		SNew(SOverlay)		
+		+SOverlay::Slot()
+		[				
+			SNew(SVerticalBox)
+			+ SVerticalBox::Slot()
+			.AutoHeight()
+			[
+				SNew(SBorder)
+				.Padding(2.0f)
+				.BorderImage(FEditorStyle::GetBrush("ToolPanel.GroupBorder"))
+				[
+					SNew(SHorizontalBox)
+					+ SHorizontalBox::Slot()
+					.Padding(3.0f, 3.0f)
+					.AutoWidth()
+					[
+						SAssignNew(AddModifierCombobox, SComboButton)
+						.OnGetMenuContent(GetContent)
+						.ButtonContent()
+						[
+							SNew(STextBlock)
+							.Text(LOCTEXT("AnimationModifierWindow_AddModifier", "Add Modifier"))
+						]
+					]
+				]
+			]
+
+			+ SVerticalBox::Slot()
+			[
+				SNew(SBorder)
+				.Padding(2.0f)
+				.BorderImage(FEditorStyle::GetBrush("ToolPanel.GroupBorder"))
+				[				
+					SNew(SSplitter)
+					.Orientation(EOrientation::Orient_Vertical)
+
+					+ SSplitter::Slot()
+					.Value(.5f)
+					[
+						SNew(SBox)
+						.Padding(2.0f)
+						[
+							SAssignNew(ModifierListView, SModifierListView)
+							.Items(&ModifierItems)
+							.InstanceDetailsView(ModifierInstanceDetailsView)
+							.OnRemoveModifier(FOnModifierArray::CreateSP(this, &SAnimationModifierContentBrowserWindow::RemoveModifiersCallback))
+						]
+					]
+
+					+ SSplitter::Slot()
+					.Value(.5f)
+					[	
+						SNew(SBox)
+						.Padding(2.0f)
+						[
+							ModifierInstanceDetailsView->AsShared()
+						]
+					]
+				]
+			]	
+
+			+ SVerticalBox::Slot()
+			.AutoHeight()
+			.HAlign(HAlign_Right)
+			.Padding(2)
+			[
+				SNew(SUniformGridPanel)
+				.SlotPadding(2)
+				+ SUniformGridPanel::Slot(0, 0)
+				[
+					SNew(SButton)
+					.HAlign(HAlign_Center)
+					.Text(LOCTEXT("AnimationModifierWindow_Import", "Apply"))
+					.ToolTipText(LOCTEXT("AnimationModifierWindow_Import_ToolTip", "Apply adding modifiers(s)."))
+					.IsEnabled(this, &SAnimationModifierContentBrowserWindow::CanApply)
+					.OnClicked(this, &SAnimationModifierContentBrowserWindow::OnApply)
+				]
+				+ SUniformGridPanel::Slot(1, 0)
+				[
+					SNew(SButton)
+					.HAlign(HAlign_Center)
+					.Text(LOCTEXT("AnimationModifierWindow_Cancel", "Cancel"))
+					.ToolTipText(LOCTEXT("AnimationModifierWindow_Cancel_ToolTip", "Cancels adding modifiers(s)."))
+					.OnClicked(this, &SAnimationModifierContentBrowserWindow::OnCancel)
+				]
+			]	
+		]
+	];
+}
+
+FReply SAnimationModifierContentBrowserWindow::OnKeyDown(const FGeometry& MyGeometry, const FKeyEvent& InKeyEvent)
+{
+	if (InKeyEvent.GetKey() == EKeys::Escape)
+	{
+		return OnCancel();
+	}
+
+	return FReply::Unhandled();
+}
+
+void SAnimationModifierContentBrowserWindow::RemoveModifiersCallback(const TArray<TWeakObjectPtr<UAnimationModifier>>& ModifiersToRemove)
+{
+	Modifiers.RemoveAll([&ModifiersToRemove](UAnimationModifier* Modifier) { return ModifiersToRemove.Contains(Modifier); });
+	ModifierItems.RemoveAll([&ModifiersToRemove](TSharedPtr<FModifierListviewItem> ModifierItem) { return ModifiersToRemove.Contains(ModifierItem->Instance); });
+	ModifierListView->Refresh();
+}
+
+void SAnimationModifierContentBrowserWindow::OnModifierPicked(UClass* PickedClass)
+{	
+	UAnimationModifier* Processor = FAnimationModifierHelpers::CreateModifierInstance(GetTransientPackage(), PickedClass);
+
+	Modifiers.Add(Processor);
+
+	FModifierListviewItem* Item = new FModifierListviewItem();
+	Item->Instance = Processor;
+	Item->Class = Processor->GetClass();
+	Item->Index = Modifiers.Num() - 1;
+	Item->OuterClass = nullptr;
+	ModifierItems.Add(ModifierListviewItem(Item));
+
+	// Close the combo box
+	AddModifierCombobox->SetIsOpen(false);
+
+	ModifierListView->Refresh();
+}
+
+void SAnimationModifierContentBrowserWindow::CreateInstanceDetailsView()
+{
+	// Create a property view
+	FPropertyEditorModule& EditModule = FModuleManager::Get().GetModuleChecked<FPropertyEditorModule>("PropertyEditor");
+
+	FDetailsViewArgs DetailsViewArgs(
+		/*bUpdateFromSelection=*/ false,
+		/*bLockable=*/ false,
+		/*bAllowSearch=*/ false,
+		FDetailsViewArgs::HideNameArea,
+		/*bHideSelectionTip=*/ true,
+		/*InNotifyHook=*/ nullptr,
+		/*InSearchInitialKeyFocus=*/ false,
+		/*InViewIdentifier=*/ NAME_None);
+	DetailsViewArgs.DefaultsOnlyVisibility = EEditDefaultsOnlyNodeVisibility::Automatic;
+	DetailsViewArgs.bShowOptions = false;
+
+	ModifierInstanceDetailsView = EditModule.CreateDetailView(DetailsViewArgs);
+	ModifierInstanceDetailsView->SetDisableCustomDetailLayouts(true);
+}
+
+FReply SAnimationModifierContentBrowserWindow::OnApply()
+{
+	const FScopedTransaction Transaction(LOCTEXT("UndoAction_ApplyModifiers", "Applying Animation Modifier(s) to Animation Sequence(s)"));
+
+	TArray<UAnimationModifiersAssetUserData*> AssetUserData;
+
+	bool bCloseWindow = true;
+
+	// Retrieve or create asset user data 
+	for (UAnimSequence* AnimationSequence : AnimSequences)
+	{
+		UAnimationModifiersAssetUserData* UserData = AnimationSequence->GetAssetUserData<UAnimationModifiersAssetUserData>();
+		if (!UserData)
+		{
+			UserData = NewObject<UAnimationModifiersAssetUserData>(AnimationSequence, UAnimationModifiersAssetUserData::StaticClass());
+			checkf(UserData, TEXT("Unable to instantiate AssetUserData class"));
+			UserData->SetFlags(RF_Transactional);
+			AnimationSequence->AddAssetUserData(UserData);
+		}
+
+		AssetUserData.Add(UserData);
+	}
+
+	// For each added modifier create add a new instance to each of the user data entries, using the one(s) set up in the window as template(s)
+	for (UAnimationModifier* Modifier : Modifiers)
+	{
+		for (UAnimationModifiersAssetUserData* UserData : AssetUserData)
+		{
+			const bool bAlreadyContainsModifier = UserData->GetAnimationModifierInstances().ContainsByPredicate([Modifier](UAnimationModifier* TestModifier) { return Modifier->GetClass() == TestModifier->GetClass(); });
+			static const FText MessageFormat = LOCTEXT("AnimationModifierWindow_AlreadyContainsModifierDialogText", "{0} already contains Animation Modifier {1}, are you sure you want to add another instance?");
+
+			const bool bUserInputResult = bAlreadyContainsModifier ? OpenMsgDlgInt(EAppMsgType::YesNo, FText::FormatOrdered(MessageFormat, FText::FromString(UserData->GetOuter()->GetName()), FText::FromString(Modifier->GetClass()->GetName())), LOCTEXT("AnimationModifierWindow_AlreadyContainsModifierTitle", "Already contains Animation Modifier!")) == EAppReturnType::Yes : true;
+			bCloseWindow = bUserInputResult;
+			if (!bAlreadyContainsModifier || bUserInputResult)
+			{
+				UObject* Outer = UserData;
+				UAnimationModifier* Processor = FAnimationModifierHelpers::CreateModifierInstance(Outer, Modifier->GetClass(), Modifier);
+				UserData->Modify();
+				UserData->AddAnimationModifier(Processor);
+			}
+		}
+	}
+
+	/** For each user data retrieve all modifiers and apply them */
+	for (int32 Index = 0; Index < AssetUserData.Num(); ++Index)
+	{
+		UAnimationModifiersAssetUserData* UserData = AssetUserData[Index];		
+		if (UserData)
+		{
+			UAnimSequence* AnimSequence = AnimSequences[Index];
+			AnimSequence->Modify();
+
+			const TArray<UAnimationModifier*>& ModifierInstances = UserData->GetAnimationModifierInstances();
+			for (UAnimationModifier* Modifier : ModifierInstances)
+			{
+				Modifier->ApplyToAnimationSequence(AnimSequence);
+			}
+		}
+	}
+
+	if (bCloseWindow && WidgetWindow.IsValid())
+	{
+		WidgetWindow.Pin()->RequestDestroyWindow();
+	}
+
+	return FReply::Handled();
+}
+
+FReply SAnimationModifierContentBrowserWindow::OnCancel()
+{
+	if (WidgetWindow.IsValid())
+	{
+		WidgetWindow.Pin()->RequestDestroyWindow();
+	}
+	return FReply::Handled();
+}
+
+bool SAnimationModifierContentBrowserWindow::CanApply() const
+{
+	return Modifiers.Num() > 0;
+}
+
+#undef LOCTEXT_NAMESPACE // "AnimationModifierContentBrowserWindow"
\ No newline at end of file
diff --git a/Engine/Source/Editor/AnimationModifiers/Private/SAnimationModifierContentBrowserWindow.h b/Engine/Source/Editor/AnimationModifiers/Private/SAnimationModifierContentBrowserWindow.h
new file mode 100644
--- /dev/null
+++ b/Engine/Source/Editor/AnimationModifiers/Private/SAnimationModifierContentBrowserWindow.h
@@ -0,0 +1,67 @@
+// Copyright 1998-2019 Epic Games, Inc. All Rights Reserved.
+
+#pragma once
+
+#include "Widgets/SCompoundWidget.h"
+#include "Widgets/SWindow.h"
+#include "Widgets/Input/SButton.h"
+
+#include "Templates/SharedPointer.h"
+#include "Containers/Array.h"
+
+class SButton;
+class SComboButton;
+class SModifierListView;
+class SMenuAnchor;
+class UAnimationModifier;
+class UAnimSequence;
+class IDetailsView;
+
+struct FModifierListviewItem;
+
+/** UI slate widget allowing the user to add Animation Modifier(s) to a selection of Animation Sequences */
+class SAnimationModifierContentBrowserWindow : public SCompoundWidget
+{
+public:
+	SLATE_BEGIN_ARGS(SAnimationModifierContentBrowserWindow)
+		: _WidgetWindow()		
+	{}			
+		SLATE_ARGUMENT(TSharedPtr<SWindow>, WidgetWindow)
+		SLATE_ARGUMENT(TArray<UAnimSequence*>, AnimSequences)
+	SLATE_END_ARGS()
+
+public:
+	SAnimationModifierContentBrowserWindow() {}
+	void Construct(const FArguments& InArgs);
+	virtual bool SupportsKeyboardFocus() const override { return true; }
+	virtual FReply OnKeyDown(const FGeometry& MyGeometry, const FKeyEvent& InKeyEvent) override;
+protected:
+	/** Callback for when the user picks a specific animation modifier class */
+	void OnModifierPicked(UClass* PickedClass);
+	/** Callback for when the user wants to remove modifier(s) from the listview */
+	void RemoveModifiersCallback(const TArray<TWeakObjectPtr<UAnimationModifier>>& ModifiersToRemove);
+
+	/** Creates the details view widget used to show Animation Modifier object details */
+	void CreateInstanceDetailsView();
+
+	/** Button callback, this applies all currently set up Animation Modifiers to the previously selected Animation Sequences */
+	FReply OnApply();
+	/** Button callback, closes the dialog/window */
+	FReply OnCancel();
+	/** Check to see whether or not the user can apply the modifiers in the current state */
+	bool CanApply() const;
+private:
+	/** Window owning this window */
+	TWeakPtr<SWindow> WidgetWindow;
+	
+	TSharedPtr<IDetailsView> ModifierInstanceDetailsView;
+	TSharedPtr<SMenuAnchor> AddModifierCombobox;
+	TSharedPtr<SModifierListView> ModifierListView;
+
+	/** Data structures used by the Modifier List View widget */
+	TArray<TSharedPtr<FModifierListviewItem>> ModifierItems;	
+	/** Current set of Animation Modifiers that would be added during Apply */
+	TArray<UAnimationModifier*> Modifiers;
+	/** Previously user-selected Animation Sequences */
+	TArray<UAnimSequence*> AnimSequences;
+};
\ No newline at end of file
diff --git a/Engine/Source/Editor/AnimationModifiers/Private/SAnimationModifiersTab.cpp b/Engine/Source/Editor/AnimationModifiers/Private/SAnimationModifiersTab.cpp
--- a/Engine/Source/Editor/AnimationModifiers/Private/SAnimationModifiersTab.cpp
+++ b/Engine/Source/Editor/AnimationModifiers/Private/SAnimationModifiersTab.cpp
@@ -25,24 +25,10 @@
 #include "ScopedTransaction.h"
 #include "Widgets/Input/SMenuAnchor.h"
 #include "Engine/BlueprintGeneratedClass.h"
+#include "AnimationModifierHelpers.h"
 
 #define LOCTEXT_NAMESPACE "SAnimationModifiersTab"
 
-/** ClassViewerFilter for Animation Modifier classes */
-class FModifierClassFilter : public IClassViewerFilter
-{
-public:
-	bool IsClassAllowed(const FClassViewerInitializationOptions& InInitOptions, const UClass* InClass, TSharedRef< FClassViewerFilterFuncs > InFilterFuncs) override
-	{
-		return InClass->IsChildOf(UAnimationModifier::StaticClass());
-	}
-
-	virtual bool IsUnloadedClassAllowed(const FClassViewerInitializationOptions& InInitOptions, const TSharedRef< const IUnloadedBlueprintData > InClass, TSharedRef< FClassViewerFilterFuncs > InFilterFuncs) override
-	{
-		return InClass->IsChildOf(UAnimationModifier::StaticClass());
-	}
-};
-
 SAnimationModifiersTab::SAnimationModifiersTab()
 	: Skeleton(nullptr), AnimationSequence(nullptr), AssetUserData(nullptr), bDirty(false)
 {	
@@ -58,16 +44,6 @@ SAnimationModifiersTab::~SAnimationModifiersTab()
 	FAssetEditorManager::Get().OnAssetOpenedInEditor().RemoveAll(this);	
 }
 
-UAnimationModifier* SAnimationModifiersTab::CreateModifierInstance(UObject* Outer, UClass* InClass)
-{
-	checkf(Outer, TEXT("Invalid outer value for modifier instantiation"));
-	UAnimationModifier* ProcessorInstance = NewObject<UAnimationModifier>(Outer, InClass);
-	checkf(ProcessorInstance, TEXT("Unable to instantiate modifier class"));
-	ProcessorInstance->SetFlags(RF_Transactional);
-	return ProcessorInstance;
-}
-
-
 void SAnimationModifiersTab::Tick(const FGeometry& AllottedGeometry, const double InCurrentTime, const float InDeltaTime)
 {
 	if (bDirty)
@@ -97,6 +73,12 @@ void SAnimationModifiersTab::Construct(const FArguments& InArgs)
 	RetrieveModifierData();
 		
 	CreateInstanceDetailsView();
+
+	FOnGetContent GetContent = FOnGetContent::CreateLambda(
+		[this]()
+	{
+		return FAnimationModifierHelpers::GetModifierPicker(FOnClassPicked::CreateRaw(this, &SAnimationModifiersTab::OnModifierPicked));
+	});
 	
 	this->ChildSlot
 	[
@@ -117,7 +99,7 @@ void SAnimationModifiersTab::Construct(const FArguments& InArgs)
 					.AutoWidth()
 					[
 						SAssignNew(AddModifierCombobox, SComboButton)
-						.OnGetMenuContent(this, &SAnimationModifiersTab::GetModifierPicker)
+						.OnGetMenuContent(GetContent)
 						.ButtonContent()
 						[
 							SNew(STextBlock)
@@ -193,36 +175,12 @@ void SAnimationModifiersTab::Construct(const FArguments& InArgs)
 	FAssetEditorManager::Get().OnAssetOpenedInEditor().AddSP(this, &SAnimationModifiersTab::OnAssetOpened);
 }
 
-
-TSharedRef<SWidget> SAnimationModifiersTab::GetModifierPicker()
-{
-	FClassViewerInitializationOptions Options;
-	Options.bShowUnloadedBlueprints = true;
-	Options.bShowNoneOption = false;	
-	TSharedPtr<FModifierClassFilter> ClassFilter = MakeShareable(new FModifierClassFilter);
-	Options.ClassFilter = ClassFilter;
-
-	FOnClassPicked OnPicked(FOnClassPicked::CreateRaw(this, &SAnimationModifiersTab::OnModifierPicked));
-
-	return SNew(SBox)
-		.WidthOverride(280)
-		[
-			SNew(SVerticalBox)
-			+ SVerticalBox::Slot()
-			.AutoHeight()
-			.MaxHeight(500)
-			[
-				FModuleManager::LoadModuleChecked<FClassViewerModule>("ClassViewer").CreateClassViewer(Options, OnPicked)
-			]
-		];
-}
-
 void SAnimationModifiersTab::OnModifierPicked(UClass* PickedClass)
 {
 	FScopedTransaction Transaction(LOCTEXT("AddModifierTransaction", "Adding Animation Modifier"));
 	
 	UObject* Outer = AssetUserData;
-	UAnimationModifier* Processor = CreateModifierInstance(Outer, PickedClass);	
+	UAnimationModifier* Processor = FAnimationModifierHelpers::CreateModifierInstance(Outer, PickedClass);
 	AssetUserData->Modify();
 	AssetUserData->AddAnimationModifier(Processor);
 	
@@ -403,7 +361,7 @@ void SAnimationModifiersTab::OnOpenModifier(const TWeakObjectPtr<UAnimationModif
 		if (Blueprint)
 		{
 			FAssetEditorManager::Get().OpenEditorForAsset(Blueprint);
-		}		
+		}
 	}
 }
 
diff --git a/Engine/Source/Editor/AnimationModifiers/Private/SAnimationModifiersTab.h b/Engine/Source/Editor/AnimationModifiers/Private/SAnimationModifiersTab.h
--- a/Engine/Source/Editor/AnimationModifiers/Private/SAnimationModifiersTab.h
+++ b/Engine/Source/Editor/AnimationModifiers/Private/SAnimationModifiersTab.h
@@ -37,8 +37,6 @@ public:
 	virtual void PostRedo(bool bSuccess) override;
 	/** End FEditorUndoClient */
 protected:
-	/** Creates the class picker for the available modifiers */
-	TSharedRef<SWidget> GetModifierPicker();
 	/** Callback for when user has picked a modifier to add */
 	void OnModifierPicked(UClass* PickedClass);
 
@@ -70,8 +68,7 @@ protected:
 
 	/** Retrieves the currently opened animation asset type and modifier user data */
 	void RetrieveAnimationAsset();
-	/** Creates a new Modifier instance to store with the current asset */
-	UAnimationModifier* CreateModifierInstance(UObject* Outer, UClass* InClass);
+
 	/** Retrieves all animation sequences which are dependent on the current opened skeleton */
 	void FindAnimSequencesForSkeleton(TArray<UAnimSequence *> &ReferencedAnimSequences);
 protected:
diff --git a/Engine/Source/Editor/AnimationModifiers/Private/SModifierListview.cpp b/Engine/Source/Editor/AnimationModifiers/Private/SModifierListview.cpp
--- a/Engine/Source/Editor/AnimationModifiers/Private/SModifierListview.cpp
+++ b/Engine/Source/Editor/AnimationModifiers/Private/SModifierListview.cpp
@@ -138,22 +138,31 @@ TSharedPtr<SWidget> SModifierListView::OnContextMenuOpening()
 		
 		MenuBuilder.BeginSection(NAME_None);
 		{		
-			if (Listview->GetNumItemsSelected() == 1)
+			if (Listview->GetNumItemsSelected() == 1 && OnOpenModifierDelegate.IsBound())
 			{
 				MenuBuilder.AddMenuEntry(LOCTEXT("OpenModifierLabel", "Open Blueprint"), LOCTEXT("OpenModifierToolTip", "Open selected Modifier Blueprint"), FSlateIcon(FEditorStyle::GetStyleSetName(), "ClassIcon.Blueprint"), FUIAction(FExecuteAction::CreateSP(this, &SModifierListView::OnOpenModifier)));
 			}
 
 			const FText ApplyLabel = FText::FormatOrdered(LOCTEXT("ApplyModifierLabel", "Apply {0}|plural(one=Modifier,other=Modifiers)"), NumItems);
 			const FText ApplyTooltip = FText::FormatOrdered(LOCTEXT("ApplyModifierToolTip", "Apply selected {0}|plural(one=Modifier,other=Modifiers)"), NumItems);
-			MenuBuilder.AddMenuEntry(ApplyLabel, ApplyTooltip, FSlateIcon(FEditorStyle::GetStyleSetName(), "GenericCommands.Redo"), FUIAction(FExecuteAction::CreateSP(this, &SModifierListView::OnApplyModifier)));	   
-
+			if (OnApplyModifierDelegate.IsBound())
+			{
+				MenuBuilder.AddMenuEntry(ApplyLabel, ApplyTooltip, FSlateIcon(FEditorStyle::GetStyleSetName(), "GenericCommands.Redo"), FUIAction(FExecuteAction::CreateSP(this, &SModifierListView::OnApplyModifier)));
+			}
+			
 			const FText RevertLabel = FText::FormatOrdered(LOCTEXT("ApplyRevertLabel", "Revert {0}|plural(one=Modifier,other=Modifiers)"), NumItems);
 			const FText RevertTooltip = FText::FormatOrdered(LOCTEXT("ApplyRevertToolTip", "Revert selected {0}|plural(one=Modifier,other=Modifiers)"), NumItems);
-			MenuBuilder.AddMenuEntry(RevertLabel, RevertTooltip, FSlateIcon(FEditorStyle::GetStyleSetName(), "GenericCommands.Undo"), FUIAction(FExecuteAction::CreateSP(this, &SModifierListView::OnRevertModifier)));
+			if (OnRevertModifierDelegate.IsBound())
+			{
+				MenuBuilder.AddMenuEntry(RevertLabel, RevertTooltip, FSlateIcon(FEditorStyle::GetStyleSetName(), "GenericCommands.Undo"), FUIAction(FExecuteAction::CreateSP(this, &SModifierListView::OnRevertModifier)));
+			}
 
 			const FText RemoveLabel = FText::FormatOrdered(LOCTEXT("RemoveModifierLabel", "Remove {0}|plural(one=Modifier,other=Modifiers)"), NumItems);
 			const FText RemoveTooltip = FText::FormatOrdered(LOCTEXT("RemoveModifierToolTip", "Remove selected {0}|plural(one=Modifier,other=Modifiers)"), NumItems);
-			MenuBuilder.AddMenuEntry(RemoveLabel, RemoveTooltip, FSlateIcon(FEditorStyle::GetStyleSetName(), "ContentBrowser.AssetActions.Delete"), FUIAction(FExecuteAction::CreateSP(this, &SModifierListView::OnRemoveModifier)));
+			if (OnRemoveModifierDelegate.IsBound())
+			{
+				MenuBuilder.AddMenuEntry(RemoveLabel, RemoveTooltip, FSlateIcon(FEditorStyle::GetStyleSetName(), "ContentBrowser.AssetActions.Delete"), FUIAction(FExecuteAction::CreateSP(this, &SModifierListView::OnRemoveModifier)));
+			}
 		}
 		MenuBuilder.EndSection();
 
@@ -161,9 +170,15 @@ TSharedPtr<SWidget> SModifierListView::OnContextMenuOpening()
 		{ 
 			MenuBuilder.BeginSection(NAME_None);
 			{
-				MenuBuilder.AddMenuEntry(LOCTEXT("MoveUpModifierLabel", "Move Up"), LOCTEXT("MoveUpModifierToolTip", "Move selected Modifier Up in list"), FSlateIcon(), FUIAction(FExecuteAction::CreateSP(this, &SModifierListView::OnMoveUpModifier), FCanExecuteAction::CreateSP( this, &SModifierListView::CanMoveSelectedItemUp)));
+				if (OnMoveUpModifierDelegate.IsBound())
+				{
+					MenuBuilder.AddMenuEntry(LOCTEXT("MoveUpModifierLabel", "Move Up"), LOCTEXT("MoveUpModifierToolTip", "Move selected Modifier Up in list"), FSlateIcon(), FUIAction(FExecuteAction::CreateSP(this, &SModifierListView::OnMoveUpModifier), FCanExecuteAction::CreateSP(this, &SModifierListView::CanMoveSelectedItemUp)));
+				}
 				
-				MenuBuilder.AddMenuEntry(LOCTEXT("MoveDownModifierLabel", "Move Down"), LOCTEXT("MoveDownModifierToolTip", "Move selected Modifier Down in list"), FSlateIcon(), FUIAction(FExecuteAction::CreateSP(this, &SModifierListView::OnMoveDownModifier), FCanExecuteAction::CreateSP(this, &SModifierListView::CanMoveSelectedItemDown)));
+				if (OnMoveDownModifierDelegate.IsBound())
+				{
+					MenuBuilder.AddMenuEntry(LOCTEXT("MoveDownModifierLabel", "Move Down"), LOCTEXT("MoveDownModifierToolTip", "Move selected Modifier Down in list"), FSlateIcon(), FUIAction(FExecuteAction::CreateSP(this, &SModifierListView::OnMoveDownModifier), FCanExecuteAction::CreateSP(this, &SModifierListView::CanMoveSelectedItemDown)));
+				}
 			}
 			MenuBuilder.EndSection();
 		}
diff --git a/Engine/Source/Editor/AnimationModifiers/Public/IAnimationModifiersModule.h b/Engine/Source/Editor/AnimationModifiers/Public/IAnimationModifiersModule.h
--- a/Engine/Source/Editor/AnimationModifiers/Public/IAnimationModifiersModule.h
+++ b/Engine/Source/Editor/AnimationModifiers/Public/IAnimationModifiersModule.h
@@ -3,7 +3,13 @@
 #pragma once
 
 #include "Modules/ModuleInterface.h"
+#include "Containers/Array.h"
+
+class UAnimSequence;
 
 class IAnimationModifiersModule : public IModuleInterface
 {
+public:
+	/** Shows a new modal dialog allowing the user to setup Animation Modifiers to be added for all AnimSequences part of InSequences */
+	virtual void ShowAddAnimationModifierWindow(const TArray<UAnimSequence*>& InSequences) = 0;
 };
