Adding missing pragma once

diff --git a/Engine/Source/Editor/AnimationModifiers/Private/AnimationModifierHelpers.h b/Engine/Source/Editor/AnimationModifiers/Private/AnimationModifierHelpers.h
--- a/Engine/Source/Editor/AnimationModifiers/Private/AnimationModifierHelpers.h
+++ b/Engine/Source/Editor/AnimationModifiers/Private/AnimationModifierHelpers.h
@@ -1,5 +1,7 @@
 // Copyright 1998-2019 Epic Games, Inc. All Rights Reserved.
 
+#pragma once
+
 #include "ClassViewerModule.h"
 #include "Widgets/SWidget.h"
 #include "Widgets/Layout/SBox.h"
