PR #5479: Fixed a couple of issues with animation notify functions in AnimationBlueprintLibrary (Contributed by [REDACTED_PERSON])

diff --git a/Engine/Source/Editor/AnimationModifiers/Private/AnimationBlueprintLibrary.cpp b/Engine/Source/Editor/AnimationModifiers/Private/AnimationBlueprintLibrary.cpp
--- a/Engine/Source/Editor/AnimationModifiers/Private/AnimationBlueprintLibrary.cpp
+++ b/Engine/Source/Editor/AnimationModifiers/Private/AnimationBlueprintLibrary.cpp
@@ -489,7 +489,7 @@ void UAnimationBlueprintLibrary::GetAnimationNotifyEventNames(const UAnimSequenc
 	}	
 }
 
-UAnimNotify* UAnimationBlueprintLibrary::AddAnimationNotifyEvent(UAnimSequence* AnimationSequence, FName NotifyTrackName, float StartTime, float Duration, TSubclassOf<UObject> NotifyClass)
+UAnimNotify* UAnimationBlueprintLibrary::AddAnimationNotifyEvent(UAnimSequence* AnimationSequence, FName NotifyTrackName, float StartTime, TSubclassOf<UAnimNotify> NotifyClass)
 {
 	UAnimNotify* Notify = nullptr;
 	if (AnimationSequence)
@@ -499,19 +499,76 @@ UAnimNotify* UAnimationBlueprintLibrary::AddAnimationNotifyEvent(UAnimSequence*
 
 		if (bIsValidTrackName && bIsValidTime)
 		{
-			AnimationSequence->Notifies.AddDefaulted();
-			FAnimNotifyEvent& NewEvent = AnimationSequence->Notifies.Last();
+			FAnimNotifyEvent& NewEvent = AnimationSequence->Notifies.AddDefaulted_GetRef();
 
 			NewEvent.NotifyName = NAME_None;
 			NewEvent.Link(AnimationSequence, StartTime);
 			NewEvent.TriggerTimeOffset = GetTriggerTimeOffsetForType(AnimationSequence->CalculateOffsetForNotify(StartTime));
 			NewEvent.TrackIndex = GetTrackIndexForAnimationNotifyTrackName(AnimationSequence, NotifyTrackName);
+			NewEvent.NotifyStateClass = nullptr;
 
 			if (NotifyClass)
 			{
-				class UObject* AnimNotifyClass = NewObject<UObject>(AnimationSequence, NotifyClass, NAME_None, RF_Transactional);
-				NewEvent.NotifyStateClass = Cast<UAnimNotifyState>(AnimNotifyClass);
-				NewEvent.Notify = Cast<UAnimNotify>(AnimNotifyClass);
+				Notify = NewObject<UAnimNotify>(AnimationSequence, NotifyClass, NAME_None, RF_Transactional);
+				NewEvent.Notify = Notify;
+
+				// Setup name for new event
+				if(NewEvent.Notify)
+				{
+					NewEvent.NotifyName = FName(*NewEvent.Notify->GetNotifyName());
+				}
+			}
+			else
+			{
+				NewEvent.Notify = nullptr;
+			}
+
+			// Refresh all cached data
+			AnimationSequence->RefreshCacheData();
+		}
+		else
+		{
+			if (!bIsValidTrackName)
+			{
+				UE_LOG(LogAnimationBlueprintLibrary, Warning, TEXT("Animation Notify Track %s does not exist on Animation Sequence %s"), *NotifyTrackName.ToString(), *AnimationSequence->GetName());
+			}
+
+			if (!bIsValidTime)
+			{
+				UE_LOG(LogAnimationBlueprintLibrary, Warning, TEXT("%f is side of Animation Sequence %s range"), StartTime, *AnimationSequence->GetName());
+			}
+		}
+	}
+	else
+	{
+		UE_LOG(LogAnimationBlueprintLibrary, Warning, TEXT("Invalid Animation Sequence for AddAnimationNotifyEvent"));
+	}
+
+	return Notify;
+}
+
+UAnimNotifyState* UAnimationBlueprintLibrary::AddAnimationNotifyStateEvent(UAnimSequence* AnimationSequence, FName NotifyTrackName, float StartTime, float Duration, TSubclassOf<UAnimNotifyState> NotifyStateClass)
+{
+	UAnimNotifyState* NotifyState = nullptr;
+	if (AnimationSequence)
+	{
+		const bool bIsValidTrackName = IsValidAnimNotifyTrackName(AnimationSequence, NotifyTrackName);
+		const bool bIsValidTime = IsValidTimeInternal(AnimationSequence, StartTime);
+
+		if (bIsValidTrackName && bIsValidTime)
+		{
+			FAnimNotifyEvent& NewEvent = AnimationSequence->Notifies.AddDefaulted_GetRef();
+
+			NewEvent.NotifyName = NAME_None;
+			NewEvent.Link(AnimationSequence, StartTime);
+			NewEvent.TriggerTimeOffset = GetTriggerTimeOffsetForType(AnimationSequence->CalculateOffsetForNotify(StartTime));
+			NewEvent.TrackIndex = GetTrackIndexForAnimationNotifyTrackName(AnimationSequence, NotifyTrackName);
+			NewEvent.Notify = nullptr;
+
+			if (NotifyStateClass)
+			{
+				NotifyState = NewObject<UAnimNotifyState>(AnimationSequence, NotifyStateClass, NAME_None, RF_Transactional);
+				NewEvent.NotifyStateClass = NotifyState;
 
 				// Setup name and duration for new event
 				if (NewEvent.NotifyStateClass)
@@ -520,21 +577,14 @@ UAnimNotify* UAnimationBlueprintLibrary::AddAnimationNotifyEvent(UAnimSequence*
 					NewEvent.SetDuration(Duration);
 					NewEvent.EndLink.Link(AnimationSequence, NewEvent.EndLink.GetTime());
 				}
-				else if(NewEvent.Notify)
-				{
-					NewEvent.NotifyName = FName(*NewEvent.Notify->GetNotifyName());
-				}
 			}
 			else
 			{
-				NewEvent.Notify = NULL;
-				NewEvent.NotifyStateClass = NULL;
+				NewEvent.NotifyStateClass = nullptr;
 			}
 
 			// Refresh all cached data
 			AnimationSequence->RefreshCacheData();			
-
-			Notify = NewEvent.Notify;
 		}
 		else
 		{
@@ -551,10 +601,10 @@ UAnimNotify* UAnimationBlueprintLibrary::AddAnimationNotifyEvent(UAnimSequence*
 	}
 	else
 	{
-		UE_LOG(LogAnimationBlueprintLibrary, Warning, TEXT("Invalid Animation Sequence for AddAnimationNotifyEvent"));
+		UE_LOG(LogAnimationBlueprintLibrary, Warning, TEXT("Invalid Animation Sequence for AddAnimationNotifyStateEvent"));
 	}
 
-	return Notify;
+	return NotifyState;
 }
 
 void UAnimationBlueprintLibrary::AddAnimationNotifyEventObject(UAnimSequence* AnimationSequence, float StartTime, UAnimNotify* Notify, FName NotifyTrackName)
@@ -568,14 +618,13 @@ void UAnimationBlueprintLibrary::AddAnimationNotifyEventObject(UAnimSequence* An
 
 		if (bValidNotify && bValidOuter && bIsValidTrackName && bIsValidTime)
 		{
-			AnimationSequence->Notifies.AddDefaulted();
-			FAnimNotifyEvent& NewEvent = AnimationSequence->Notifies.Last();
+			FAnimNotifyEvent& NewEvent = AnimationSequence->Notifies.AddDefaulted_GetRef();
 
 			NewEvent.NotifyName = NAME_None;
 			NewEvent.Link(AnimationSequence, StartTime);
 			NewEvent.TriggerTimeOffset = GetTriggerTimeOffsetForType(AnimationSequence->CalculateOffsetForNotify(StartTime));
 			NewEvent.TrackIndex = GetTrackIndexForAnimationNotifyTrackName(AnimationSequence, NotifyTrackName);
-			NewEvent.NotifyStateClass = Cast<UAnimNotifyState>(Notify);
+			NewEvent.NotifyStateClass = nullptr;
 			NewEvent.Notify = Notify;
 
 			// Refresh all cached data
@@ -609,7 +658,61 @@ void UAnimationBlueprintLibrary::AddAnimationNotifyEventObject(UAnimSequence* An
 	{
 		UE_LOG(LogAnimationBlueprintLibrary, Warning, TEXT("Invalid Animation Sequence for AddAnimationNotifyEventObject"));
 	}
-	
+}
+
+void UAnimationBlueprintLibrary::AddAnimationNotifyStateEventObject(UAnimSequence* AnimationSequence, float StartTime, float Duration, UAnimNotifyState* NotifyState, FName NotifyTrackName)
+{
+	if (AnimationSequence)
+	{
+		const bool bValidNotifyState = NotifyState != nullptr;
+		const bool bValidOuter = bValidNotifyState && NotifyState->GetOuter() == AnimationSequence;
+		const bool bIsValidTrackName = IsValidAnimNotifyTrackName(AnimationSequence, NotifyTrackName);
+		const bool bIsValidTime = IsValidTimeInternal(AnimationSequence, StartTime);
+
+		if (bValidNotifyState && bValidOuter && bIsValidTrackName && bIsValidTime)
+		{
+			FAnimNotifyEvent& NewEvent = AnimationSequence->Notifies.AddDefaulted_GetRef();
+
+			NewEvent.NotifyName = NAME_None;
+			NewEvent.Link(AnimationSequence, StartTime);
+			NewEvent.TriggerTimeOffset = GetTriggerTimeOffsetForType(AnimationSequence->CalculateOffsetForNotify(StartTime));
+			NewEvent.TrackIndex = GetTrackIndexForAnimationNotifyTrackName(AnimationSequence, NotifyTrackName);
+			NewEvent.NotifyStateClass = NotifyState;
+			NewEvent.Notify = nullptr;
+			NewEvent.SetDuration(Duration);
+			NewEvent.EndLink.Link(AnimationSequence, NewEvent.EndLink.GetTime());
+
+			// Refresh all cached data
+			AnimationSequence->RefreshCacheData();
+		}
+		else
+		{
+			if (!bValidNotifyState)
+			{
+				UE_LOG(LogAnimationBlueprintLibrary, Warning, TEXT("Invalid Animation Notify in AddAnimationNotifyEventObject"));
+			}
+
+			if (!bValidOuter)
+			{
+				const FString NotifyName = bValidNotifyState ? NotifyState->GetName() : "Invalid Notify";
+				UE_LOG(LogAnimationBlueprintLibrary, Warning, TEXT("Animation Notify %s Outer is not %s"), *NotifyName, *AnimationSequence->GetName());
+			}
+
+			if (!bIsValidTrackName)
+			{
+				UE_LOG(LogAnimationBlueprintLibrary, Warning, TEXT("Animation Notify Track %s does not exist on Animation Sequence %s"), *NotifyTrackName.ToString(), *AnimationSequence->GetName());
+			}
+
+			if (!bIsValidTime)
+			{
+				UE_LOG(LogAnimationBlueprintLibrary, Warning, TEXT("%f is side of Animation Sequence %s range"), StartTime, *AnimationSequence->GetName());
+			}
+		}
+	}
+	else
+	{
+		UE_LOG(LogAnimationBlueprintLibrary, Warning, TEXT("Invalid Animation Sequence for AddAnimationNotifyEventObject"));
+	}
 }
 
 static void ReplaceAnimNotifies_Helper(UAnimSequenceBase* AnimationSequence, UClass* OldNotifyClass, UClass* NewNotifyClass, FOnNotifyReplaced OnNotifyReplaced, FOnNotifyStateReplaced OnNotifyStateReplaced)
diff --git a/Engine/Source/Editor/AnimationModifiers/Public/AnimationBlueprintLibrary.h b/Engine/Source/Editor/AnimationModifiers/Public/AnimationBlueprintLibrary.h
--- a/Engine/Source/Editor/AnimationModifiers/Public/AnimationBlueprintLibrary.h
+++ b/Engine/Source/Editor/AnimationModifiers/Public/AnimationBlueprintLibrary.h
@@ -167,12 +167,20 @@ public:
 
 	/** Adds an Animation Notify Event to Notify track in the given Animation with the given Notify creation data */
 	UFUNCTION(BlueprintCallable, Category = "AnimationBlueprintLibrary|NotifyEvents")
-	static UAnimNotify* AddAnimationNotifyEvent(UAnimSequence* AnimationSequence, FName NotifyTrackName, float StartTime, float Duration, TSubclassOf<UObject> NotifyClass);
+	static UAnimNotify* AddAnimationNotifyEvent(UAnimSequence* AnimationSequence, FName NotifyTrackName, float StartTime, TSubclassOf<UAnimNotify> NotifyClass);
+
+	/** Adds an Animation Notify State Event to Notify track in the given Animation with the given Notify State creation data */
+	UFUNCTION(BlueprintCallable, Category = "AnimationBlueprintLibrary|NotifyEvents")
+	static UAnimNotifyState* AddAnimationNotifyStateEvent(UAnimSequence* AnimationSequence, FName NotifyTrackName, float StartTime, float Duration, TSubclassOf<UAnimNotifyState> NotifyStateClass);
 
 	/** Adds an the specific Animation Notify to the Animation Sequence (requires Notify's outer to be the Animation Sequence) */
 	UFUNCTION(BlueprintCallable, Category = "AnimationBlueprintLibrary|NotifyEvents")
 	static void AddAnimationNotifyEventObject(UAnimSequence* AnimationSequence, float StartTime, UAnimNotify* Notify, FName NotifyTrackName);
 
+	/** Adds an the specific Animation Notify State to the Animation Sequence (requires Notify State's outer to be the Animation Sequence) */
+	UFUNCTION(BlueprintCallable, Category = "AnimationBlueprintLibrary|NotifyEvents")
+	static void AddAnimationNotifyStateEventObject(UAnimSequence* AnimationSequence, float StartTime, float Duration, UAnimNotifyState* NotifyState, FName NotifyTrackName);
+
 	/** Removes Animation Notify Events found by Name within the Animation Sequence, and returns the number of removed name instances */
 	UFUNCTION(BlueprintCallable, Category = "AnimationBlueprintLibrary|NotifyEvents")
 	static int32 RemoveAnimationNotifyEventsByName(UAnimSequence* AnimationSequence, FName NotifyName);
