Add Curve via Anim Modifiers not working
#fix Make sure we call MarkRawDataAsModified, to ensure the curve data is updated

diff --git a/Engine/Source/Editor/AnimationModifiers/Private/AnimationModifier.cpp b/Engine/Source/Editor/AnimationModifiers/Private/AnimationModifier.cpp
--- a/Engine/Source/Editor/AnimationModifiers/Private/AnimationModifier.cpp
+++ b/Engine/Source/Editor/AnimationModifiers/Private/AnimationModifier.cpp
@@ -73,9 +73,10 @@ void UAnimationModifier::ApplyToAnimationSequence(class UAnimSequence* InAnimati
 
 
 		CurrentAnimSequence->PostEditChange();
-		CurrentSkeleton->PostEditChange();		
+		CurrentSkeleton->PostEditChange();
 		CurrentAnimSequence->RefreshCacheData();
 		CurrentAnimSequence->RefreshCurveData();
+		CurrentAnimSequence->MarkRawDataAsModified();
 
 		UpdateStoredRevisions();
 	}
@@ -126,6 +127,7 @@ void UAnimationModifier::RevertFromAnimationSequence(class UAnimSequence* InAnim
 	CurrentSkeleton->PostEditChange();
 	CurrentAnimSequence->RefreshCacheData();
 	CurrentAnimSequence->RefreshCurveData();
+	CurrentAnimSequence->MarkRawDataAsModified();
 
 	ResetStoredRevisions();
 
