Make Animation Modifiers retain the previously applied version, this way when behaviour or property values change the reverting stage is actually correct.

diff --git a/Engine/Source/Editor/AnimationModifiers/Private/AnimationModifier.cpp b/Engine/Source/Editor/AnimationModifiers/Private/AnimationModifier.cpp
--- a/Engine/Source/Editor/AnimationModifiers/Private/AnimationModifier.cpp
+++ b/Engine/Source/Editor/AnimationModifiers/Private/AnimationModifier.cpp
@@ -10,6 +10,12 @@
 #include "Dialogs/Dialogs.h"
 #include "Editor/Transactor.h"
 #include "UObject/UObjectIterator.h"
+#include "UObject/AnimObjectVersion.h"
+
+UAnimationModifier::UAnimationModifier()
+	: PreviouslyAppliedModifier(nullptr)
+{
+}
 
 void UAnimationModifier::ApplyToAnimationSequence(class UAnimSequence* InAnimationSequence)
 {
@@ -34,8 +40,15 @@ void UAnimationModifier::ApplyToAnimationSequence(class UAnimSequence* InAnimati
 	AnimationDataTransaction.SaveObject(CurrentAnimSequence);
 	AnimationDataTransaction.SaveObject(CurrentSkeleton);
 
-	/** This populates the log with possible warnings and or errors to notify the user about */
-	OnRevert(CurrentAnimSequence);
+	/** In case this modifier has been previously applied, revert it using the serialised out version at the time */	
+	if (PreviouslyAppliedModifier)
+	{
+
+		PreviouslyAppliedModifier->Modify();
+		PreviouslyAppliedModifier->OnRevert(CurrentAnimSequence);
+	}
+
+	/** Reverting and applying, populates the log with possible warnings and or errors to notify the user about */
 	OnApply(CurrentAnimSequence);
 
 	// Apply transaction
@@ -70,7 +83,14 @@ void UAnimationModifier::ApplyToAnimationSequence(class UAnimSequence* InAnimati
 	else
 	{		
 		UpdateCompressedAnimationData();
+		
+		/** Mark the previous modifier pending kill, as it will be replaced with the current modifier state */
+		if (PreviouslyAppliedModifier)
+		{
+			PreviouslyAppliedModifier->MarkPendingKill();
+		}
 
+		PreviouslyAppliedModifier = DuplicateObject(this, GetOuter());
 
 		CurrentAnimSequence->PostEditChange();
 		CurrentSkeleton->PostEditChange();
@@ -106,34 +126,42 @@ void UAnimationModifier::RevertFromAnimationSequence(class UAnimSequence* InAnim
 {
 	FEditorScriptExecutionGuard ScriptGuard;
 
-	checkf(InAnimationSequence, TEXT("Invalid Animation Sequence supplied"));
-	CurrentAnimSequence = InAnimationSequence;
-	CurrentSkeleton = InAnimationSequence->GetSkeleton();
+	/** Can only revert if previously applied, which means there should be a previous modifier */
+	if (PreviouslyAppliedModifier)
+	{
+		checkf(InAnimationSequence, TEXT("Invalid Animation Sequence supplied"));
+		CurrentAnimSequence = InAnimationSequence;
+		CurrentSkeleton = InAnimationSequence->GetSkeleton();
 
-	// Transact the modifier to prevent instance variables/data to change during reverting
-	FTransaction Transaction;
-	Transaction.SaveObject(this);
+		// Transact the modifier to prevent instance variables/data to change during reverting
+		FTransaction Transaction;
+		Transaction.SaveObject(this);
 
-	OnRevert(CurrentAnimSequence);
+		PreviouslyAppliedModifier->Modify();
+		PreviouslyAppliedModifier->OnRevert(CurrentAnimSequence);
 
-	// Apply transaction
-	Transaction.BeginOperation();
-	Transaction.Apply();
-	Transaction.EndOperation();
+		// Apply transaction
+		Transaction.BeginOperation();
+		Transaction.Apply();
+		Transaction.EndOperation();
 
-	UpdateCompressedAnimationData();
+		UpdateCompressedAnimationData();
 
-	CurrentAnimSequence->PostEditChange();
-	CurrentSkeleton->PostEditChange();
-	CurrentAnimSequence->RefreshCacheData();
-	CurrentAnimSequence->RefreshCurveData();
-	CurrentAnimSequence->MarkRawDataAsModified();
+	    CurrentAnimSequence->PostEditChange();
+	    CurrentSkeleton->PostEditChange();
+	    CurrentAnimSequence->RefreshCacheData();
+	    CurrentAnimSequence->RefreshCurveData();
+	    CurrentAnimSequence->MarkRawDataAsModified();
 
-	ResetStoredRevisions();
+		ResetStoredRevisions();
 
-	// Finished
-	CurrentAnimSequence = nullptr;
-	CurrentSkeleton = nullptr;
+		// Finished
+		CurrentAnimSequence = nullptr;
+		CurrentSkeleton = nullptr;
+
+		PreviouslyAppliedModifier->MarkPendingKill();
+		PreviouslyAppliedModifier = nullptr;
+	}
 }
 
 bool UAnimationModifier::IsLatestRevisionApplied() const
@@ -154,6 +182,18 @@ void UAnimationModifier::PostInitProperties()
 	}
 }
 
+void UAnimationModifier::Serialize(FArchive& Ar)
+{
+	Super::Serialize(Ar);
+	Ar.UsingCustomVersion(FReleaseObjectVersion::GUID);
+
+	/** Back-wards compatibility, assume the current modifier as previously applied */
+	if (Ar.CustomVer(FReleaseObjectVersion::GUID) < FReleaseObjectVersion::SerializeAnimModifierState)
+	{
+		PreviouslyAppliedModifier = DuplicateObject(this, GetOuter());
+	}
+}
+
 const USkeleton* UAnimationModifier::GetSkeleton()
 {
 	return CurrentSkeleton;
diff --git a/Engine/Source/Editor/AnimationModifiers/Private/SAnimationModifiersTab.cpp b/Engine/Source/Editor/AnimationModifiers/Private/SAnimationModifiersTab.cpp
--- a/Engine/Source/Editor/AnimationModifiers/Private/SAnimationModifiersTab.cpp
+++ b/Engine/Source/Editor/AnimationModifiers/Private/SAnimationModifiersTab.cpp
@@ -142,6 +142,7 @@ void SAnimationModifiersTab::Construct(const FArguments& InArgs)
 							.Items(&ModifierItems)
 							.InstanceDetailsView(ModifierInstanceDetailsView)
 							.OnApplyModifier(FOnModifierArray::CreateSP(this, &SAnimationModifiersTab::OnApplyModifier))
+							.OnCanRevertModifier(FOnCanModifierArray::CreateSP(this, &SAnimationModifiersTab::OnCanRevertModifier))
 							.OnRevertModifier(FOnModifierArray::CreateSP(this, &SAnimationModifiersTab::OnRevertModifier))
 							.OnRemoveModifier(FOnModifierArray::CreateSP(this, &SAnimationModifiersTab::OnRemoveModifier))
 							.OnOpenModifier(FOnSingleModifier::CreateSP(this, &SAnimationModifiersTab::OnOpenModifier))
@@ -317,6 +318,26 @@ void SAnimationModifiersTab::OnRevertModifier(const TArray<TWeakObjectPtr<UAnima
 	RevertModifiers(ModifierInstances);
 }
 
+bool SAnimationModifiersTab::OnCanRevertModifier(const TArray<TWeakObjectPtr<UAnimationModifier>>& Instances)
+{
+	bool bCanRevert = false;
+
+	for (TWeakObjectPtr<UAnimationModifier> InstancePtr : Instances)
+	{
+		checkf(InstancePtr.IsValid(), TEXT("Invalid weak object ptr to modifier instance"));
+		UAnimationModifier* Instance = InstancePtr.Get();
+
+		// At least one instance has to be revert-able
+		if (Instance->CanRevert())
+		{
+			bCanRevert = true;
+			break;
+		}
+	}
+
+	return bCanRevert;
+}
+
 void SAnimationModifiersTab::OnRemoveModifier(const TArray<TWeakObjectPtr<UAnimationModifier>>& Instances)
 {
 	const bool bShouldRevert = OpenMsgDlgInt(EAppMsgType::YesNo, LOCTEXT("RemoveAndRevertPopupText", "Should the Modifiers be reverted before removing them?"), FText::FromString("Revert before Removing")) == EAppReturnType::Yes;
diff --git a/Engine/Source/Editor/AnimationModifiers/Private/SAnimationModifiersTab.h b/Engine/Source/Editor/AnimationModifiers/Private/SAnimationModifiersTab.h
--- a/Engine/Source/Editor/AnimationModifiers/Private/SAnimationModifiersTab.h
+++ b/Engine/Source/Editor/AnimationModifiers/Private/SAnimationModifiersTab.h
@@ -48,6 +48,7 @@ protected:
 	/** Callbacks for available modifier actions */
 	void OnApplyModifier(const TArray<TWeakObjectPtr<UAnimationModifier>>& Instances);	
 	void OnRevertModifier(const TArray<TWeakObjectPtr<UAnimationModifier>>& Instances);
+	bool OnCanRevertModifier(const TArray<TWeakObjectPtr<UAnimationModifier>>& Instances);	
 	void OnRemoveModifier(const TArray<TWeakObjectPtr<UAnimationModifier>>& Instances);
 	void OnOpenModifier(const TWeakObjectPtr<UAnimationModifier>& Instance);
 
diff --git a/Engine/Source/Editor/AnimationModifiers/Private/SModifierListview.cpp b/Engine/Source/Editor/AnimationModifiers/Private/SModifierListview.cpp
--- a/Engine/Source/Editor/AnimationModifiers/Private/SModifierListview.cpp
+++ b/Engine/Source/Editor/AnimationModifiers/Private/SModifierListview.cpp
@@ -24,6 +24,7 @@ void SModifierListView::Construct(const FArguments& InArgs)
 
 	OnApplyModifierDelegate = InArgs._OnApplyModifier;
 	OnRevertModifierDelegate = InArgs._OnRevertModifier;
+	OnCanRevertModifierDelegate = InArgs._OnCanRevertModifier;
 	OnOpenModifierDelegate = InArgs._OnOpenModifier;
 	OnRemoveModifierDelegate = InArgs._OnRemoveModifier;
 	OnMoveUpModifierDelegate = InArgs._OnMoveUpModifier;
@@ -90,6 +91,17 @@ void SModifierListView::OnRevertModifier()
 	}
 }
 
+bool SModifierListView::OnCanRevertModifier()
+{
+	if (OnCanRevertModifierDelegate.IsBound() && Listview->GetNumItemsSelected())
+	{
+		TArray<TWeakObjectPtr<UAnimationModifier>> Instances = GetSelectedModifierInstances();
+		return OnCanRevertModifierDelegate.Execute(Instances);
+	}
+
+	return false;
+}
+
 void SModifierListView::OnMoveUpModifier()
 {
 	if (OnOpenModifierDelegate.IsBound() && (Listview->GetNumItemsSelected() == 1))
@@ -154,7 +166,7 @@ TSharedPtr<SWidget> SModifierListView::OnContextMenuOpening()
 			const FText RevertTooltip = FText::FormatOrdered(LOCTEXT("ApplyRevertToolTip", "Revert selected {0}|plural(one=Modifier,other=Modifiers)"), NumItems);
 			if (OnRevertModifierDelegate.IsBound())
 			{
-				MenuBuilder.AddMenuEntry(RevertLabel, RevertTooltip, FSlateIcon(FEditorStyle::GetStyleSetName(), "GenericCommands.Undo"), FUIAction(FExecuteAction::CreateSP(this, &SModifierListView::OnRevertModifier)));
+				MenuBuilder.AddMenuEntry(RevertLabel, RevertTooltip, FSlateIcon(FEditorStyle::GetStyleSetName(), "GenericCommands.Undo"), FUIAction(FExecuteAction::CreateSP(this, &SModifierListView::OnRevertModifier), FCanExecuteAction::CreateSP(this, &SModifierListView::OnCanRevertModifier)));
 			}
 
 			const FText RemoveLabel = FText::FormatOrdered(LOCTEXT("RemoveModifierLabel", "Remove {0}|plural(one=Modifier,other=Modifiers)"), NumItems);
diff --git a/Engine/Source/Editor/AnimationModifiers/Private/SModifierListview.h b/Engine/Source/Editor/AnimationModifiers/Private/SModifierListview.h
--- a/Engine/Source/Editor/AnimationModifiers/Private/SModifierListview.h
+++ b/Engine/Source/Editor/AnimationModifiers/Private/SModifierListview.h
@@ -7,6 +7,7 @@
 #include "AnimationModifier.h"
 
 DECLARE_DELEGATE_OneParam(FOnModifierArray, const TArray<TWeakObjectPtr<UAnimationModifier>>&);
+DECLARE_DELEGATE_RetVal_OneParam(bool, FOnCanModifierArray, const TArray<TWeakObjectPtr<UAnimationModifier>>&);
 DECLARE_DELEGATE_OneParam(FOnSingleModifier, const TWeakObjectPtr<UAnimationModifier>&);
 
 /** Data representation of a modifier in the listview */
@@ -27,6 +28,7 @@ class SModifierListView : public SCompoundWidget
 	SLATE_ARGUMENT(TArray<ModifierListviewItem>*, Items)
 	SLATE_ARGUMENT(TSharedPtr<IDetailsView>, InstanceDetailsView)
 	SLATE_ARGUMENT(FOnModifierArray, OnApplyModifier);
+	SLATE_ARGUMENT(FOnCanModifierArray, OnCanRevertModifier)
 	SLATE_ARGUMENT(FOnModifierArray, OnRevertModifier);
 	SLATE_ARGUMENT(FOnModifierArray, OnRemoveModifier);
 	SLATE_ARGUMENT(FOnSingleModifier, OnOpenModifier);
@@ -44,6 +46,7 @@ class SModifierListView : public SCompoundWidget
 	void OnRemoveModifier();
 	void OnOpenModifier();
 	void OnRevertModifier();
+	bool OnCanRevertModifier();
 	void OnMoveUpModifier();
 	void OnMoveDownModifier();
 
@@ -65,6 +68,7 @@ protected:
 	/** Delegates back to SAnimationModifiersTab functionality */
 	FOnModifierArray OnApplyModifierDelegate;
 	FOnModifierArray OnRevertModifierDelegate;
+	FOnCanModifierArray OnCanRevertModifierDelegate;
 	FOnModifierArray OnRemoveModifierDelegate;
 	FOnSingleModifier OnOpenModifierDelegate;
 	FOnSingleModifier OnMoveUpModifierDelegate;
diff --git a/Engine/Source/Editor/AnimationModifiers/Public/AnimationModifier.h b/Engine/Source/Editor/AnimationModifiers/Public/AnimationModifier.h
--- a/Engine/Source/Editor/AnimationModifiers/Public/AnimationModifier.h
+++ b/Engine/Source/Editor/AnimationModifiers/Public/AnimationModifier.h
@@ -17,6 +17,8 @@ class ANIMATIONMODIFIERS_API UAnimationModifier : public UObject
 
 	friend class FAnimationModifierDetailCustomization;
 public:
+	UAnimationModifier();
+
 	/** Applying and reverting the modifier for the given Animation Sequence */
 	void ApplyToAnimationSequence(class UAnimSequence* InAnimationSequence);
 	void RevertFromAnimationSequence(class UAnimSequence* InAnimationSequence);
@@ -30,11 +32,15 @@ public:
 	void OnRevert(UAnimSequence* AnimationSequence);
 	virtual void OnRevert_Implementation(UAnimSequence* AnimationSequence) {}
 
+	/** Returns whether or not this modifier can be reverted, which means it will have to been applied (PreviouslyAppliedModifier != nullptr) */
+	bool CanRevert() const { return PreviouslyAppliedModifier; };
+
 	/** Whether or not the latest compiled version of the blueprint is applied for this instance */
 	bool IsLatestRevisionApplied() const;
 
 	// Begin UObject Overrides
 	virtual void PostInitProperties() override;
+	virtual void Serialize(FArchive& Ar) override;
 	// End UObject Overrides
 protected:
 	// Derived class accessors to skeleton and anim sequence 
@@ -69,4 +75,8 @@ private:
 	// This holds the latest value returned by UpdateNativeRevisionGuid during the last PostLoad (changing this value will invalidate the GUIDs for all instances)
 	UPROPERTY(config)
 	int32 StoredNativeRevision;
+
+	/** Serialized version of the modifier that has been previously applied to the Animation Asset */
+	UPROPERTY()
+	UAnimationModifier* PreviouslyAppliedModifier;
 };
\ No newline at end of file
diff --git a/Engine/Source/Runtime/Core/Public/UObject/ReleaseObjectVersion.h b/Engine/Source/Runtime/Core/Public/UObject/ReleaseObjectVersion.h
--- a/Engine/Source/Runtime/Core/Public/UObject/ReleaseObjectVersion.h
+++ b/Engine/Source/Runtime/Core/Public/UObject/ReleaseObjectVersion.h
@@ -87,6 +87,9 @@ struct CORE_API FReleaseObjectVersion
 		GroomAssetVersion1,
 		GroomAssetVersion2,
 
+		// Store applied version of Animation Modifier to use when reverting
+		SerializeAnimModifierState,
+
 		// -----<new versions can be added above this line>-------------------------------------------------
 		VersionPlusOne,
 		LatestVersion = VersionPlusOne - 1
