Fix implicit cast
#rnx

diff --git a/Engine/Source/Editor/AnimationModifiers/Public/AnimationModifier.h b/Engine/Source/Editor/AnimationModifiers/Public/AnimationModifier.h
--- a/Engine/Source/Editor/AnimationModifiers/Public/AnimationModifier.h
+++ b/Engine/Source/Editor/AnimationModifiers/Public/AnimationModifier.h
@@ -33,7 +33,7 @@ public:
 	virtual void OnRevert_Implementation(UAnimSequence* AnimationSequence) {}
 
 	/** Returns whether or not this modifier can be reverted, which means it will have to been applied (PreviouslyAppliedModifier != nullptr) */
-	bool CanRevert() const { return PreviouslyAppliedModifier; };
+	bool CanRevert() const { return PreviouslyAppliedModifier != nullptr; };
 
 	/** Whether or not the latest compiled version of the blueprint is applied for this instance */
 	bool IsLatestRevisionApplied() const;
