fix missing include
#rnx

diff --git a/Engine/Source/Editor/AnimationModifiers/Private/AnimationModifier.cpp b/Engine/Source/Editor/AnimationModifiers/Private/AnimationModifier.cpp
--- a/Engine/Source/Editor/AnimationModifiers/Private/AnimationModifier.cpp
+++ b/Engine/Source/Editor/AnimationModifiers/Private/AnimationModifier.cpp
@@ -7,6 +7,7 @@
 #include "Editor/Transactor.h"
 #include "UObject/UObjectIterator.h"
 
+#include "UObject/ReleaseObjectVersion.h"
 #include "Dialogs/Dialogs.h"
 #include "Editor/Transactor.h"
 #include "UObject/UObjectIterator.h"
