Fix crash when removing curve and there is no layering editing happened.

diff --git a/Engine/Source/Editor/AnimationModifiers/Private/AnimationBlueprintLibrary.cpp b/Engine/Source/Editor/AnimationModifiers/Private/AnimationBlueprintLibrary.cpp
--- a/Engine/Source/Editor/AnimationModifiers/Private/AnimationBlueprintLibrary.cpp
+++ b/Engine/Source/Editor/AnimationModifiers/Private/AnimationBlueprintLibrary.cpp
@@ -1569,7 +1569,7 @@ FName UAnimationBlueprintLibrary::RetrieveContainerNameForCurve(const UAnimSeque
 	for (int32 Index = 0; Index < (int32)ESmartNameContainerType::SNCT_MAX; ++Index)
 	{
 		const FSmartNameMapping* CurveMapping = AnimationSequence->GetSkeleton()->GetSmartNameContainer(SmartContainerNames[Index]);
-		if (CurveMapping->Exists(CurveName))
+		if (CurveMapping && CurveMapping->Exists(CurveName))
 		{
 			return SmartContainerNames[Index];
 		}
