Animation data MVC refactor

+ Introduced UAnimDataModel, this currently represents the source data for bone and curve animation. It contains:
    + Bone animation tracks (FBoneAnimationTrack)
        + Key data (coarse)
        + Name
        + Bone tree index
    + Curve data (FAnimationCurveData)
        + Float Curves
        + Transform Curves
    + Play length
    + Sampling rate
        + Used for deriving the expected number of keys/frames when combined with the PlayLength
    + Transient data for supporting backwards compatibility APIs
    + (Dynamic) delegate which broadcasts the mutation Notifies

+ Introduced UAnimDataController, this has sole authority over mutating data contained by UAnimDataModel
    + API functionality allows to transform the contained data in all ways currently expected in the engine.
        + Any mutation will add an undo/redo-able FChange object into the transaction buffer
            + FChangeTransactor helper object allows for keeping track of (compounded) FChange's and inserting them into GUndo
        + Broadcasts change notifies alongside a payload object containing details about the change
            + Interested systems/objects can register to these changes through UAnimDataModel::OnModelModified event
        + Almost all functionality is exposed to both Blueprint and Python scripting
    + Allows for opening 'Brackets', these define the beginning and end of a chain of mutations. Allowing anything registered to OnModelModified to halt responding to the notifies until the (top-level) bracket is closed
    + Undo/redo actions
        + Each mutation to the UAnimDataModel is covered by a 'Action' which is based of FChange and is used to insert a undo/redo-able operation into the transactions buffer

+ Introduced EAnimDataModelNotifType and per-notify payload types. These are used to update views, and any model derived data

+ Introduced FAnimationCurveIdentifier, exposed to scripting, this is used for referencing a curve by name and type when passed to the controller
    + Allows for targetting a specific RichCurve as part of a TransformCurve

+ Introduced FAnimationDataNotifyCollectorused for keeping track of which notifies of type EAnimDataModelNotifType are broadcasted between top-level EAnimDataModelNotifType::BracketOpened and EAnimDataModelNotifType::BracketClosed notifies

+ Added CanTransact to UEngine, returns whether or not a transaction can be made

+ Added UAnimCompositeBase::SetCompositeLength, used for runtime changing of play length value

+ Animation Sequence helpers, containing 'helper' functionality for both AnimSequence and AnimDataModel

* Animation Sequence Base
    + virtual PopulateModel, called during upgrade path for converting existing (legacy) data to the new UAnimDataModel data
    + virtual OnModelModified, registered to the Model's OnModified event to handle any Notifies and payloads
    + Added UAnimDataModel sub-object (editor only)
    + Added UAnimDataController instance (transient and editor only)
    + FAnimationDataNotifyCollector to keep track of model bracketed notifies
    * Deprecated
        * SetSequenceLength()
        * RefreshCurveData()
        * MarkRawDataAsModified()
        * RegisterOnAnimCurvesChanged()
        * UnregisterOnAnimCurvesChanged()
        * UnregisterOnAnimTrackCurvesChanged()
        * RegisterOnAnimTrackCurvesChanged()
        * Public access to RawCurveData

* Animation Sequence
    + Implements PopulateModel/OnModelModified for AnimationSequence related data
    + Added TargetFrameRate, currently locked to the initial sampling rate, but allows for resampling the UAnimDataModel data according to the set rate
    + Added NumberOfSampledKeys (editor only), populated by resampling process
    + Added NumberOfSampledFrames (editor only), populated by resampling process
    + Added ResampledAnimationTrackData (editor only and transient), populated by resampling process
    + Added bBlockCompressionRequests (editor only), used for blocking compression during automation tests
    * Deprecated
        * SetNumberOfSampledKeys()
        * MarkRawDataAsModified()
        * GetRawAnimationData()
        * GetAnimationTrackNames()
        * AddNewRawTrack()
        * GetRawTrackToSkeletonMapTable()
        * GetRawAnimationTrack()
        * GetRawAnimationTrack()
        * UpdateFrameRate()
        * ExtractBoneTransform()
        * CompressRawAnimData()
        * GetSkeletonIndexFromRawDataTrackIndex()
        * RecycleAnimSequence()
        * CleanAnimSequenceForImport()
        * CopyNotifies()
        * PostProcessSequence()
        * AddLoopingInterpolation()
        * RemoveAllTracks()
        * DoesContainTransformCurves()
        * InsertFramesToRawAnimData()
        * CropRawAnimData()
        * RemoveTrack()
        * InsertTrack()
        * ResizeSequence()
        * Non-editor access to GetUncompressedRawSize()
        * Non-editor access to GetApproxRawSize()
        * Public access to UpdateCompressedCurveName()
        * NumberOfKeys
        * SamplingFrameRate
        * TrackToSkeletonMapTable
        * RawAnimationData
        * AnimationTrackNames

* Animation Streamable
    * Model from source Sequence is duplicated rather than copying animation data

* Deprecated ERawCurveTrackTypes::RCT_Vector (marked as hidden)

+ Added automated test for
    + All script exposed controller functionality
    + Undo/redo actions
* Updated existing AnimSequence tests with deprecation and new expected data

+ Base data for compression is now populated from the 'resampled' version stored on the AnimSequence
+ Data cleanup and validation is now performed during compression
    + Track sanitization is now performed during compression (normalizing Quats and clamping near-zero values to zero)
    + Key reduction
    + Invalid track (missing bone from skeleton) removal
    + Curve name validation against skeleton
* Replaced RawCurves with Float curves
* Updated DDC key
- Removed all Guid generation related functionality (now part of the model)

* AnimSequenceBase model registers to AnimModel's notify event, used for refreshing the tracks
    * Replaces the in-place calling of RefreshTracks()
* All curve tracks now hold a const CurveType* rather than a CurveType*/& for introspection of its data
    + Any mutations now go through the controller API
    + Uses AnimationCurveIdentifier to set Transform's per-channel tracks
* RichCurveEditorModelNamed conformed to model/controller
    + Any modifications are applied to an temporary CurveType which always contains a copy of the const-ptr after any previous mutation
    + Registers to the outer AnimModel's notify event, used for updating the cached curve data
    + Registers to CurveModifiedDelegate, used for copying the temp curve data to the model using SetCurveKeys
        * Not-ideal but point that I got to for V1
        * Would be replaced with either refactoring FRichCurve to be MVC like, or by calling UAnimDataController functionality from an implementation of FRichCurveEditorModel

* Mark FChange overrides as final for swap/command change permutations
* Changed FCompoundChange GetDescription to return concatenation of all sub-changes
* Added GetDescription to FTransaction which returns the ToString() value of any contained FChange entries.
    * Entries in SUndoHistory tab now have a ToolTip showing the GetDescription() value
+ FChangeTransactor helper object allows for keeping track of (compounded) FChange's and inserting them into GUndo

* Deprecation handling, conforming code to new APIs and exposing structure/objects/properties to scripting
* Conforming AnimSequence importing from FBX to Model/Controller changes

diff --git a/Engine/Plugins/Experimental/AlembicImporter/Source/AlembicLibrary/Private/AbcImporter.cpp b/Engine/Plugins/Experimental/AlembicImporter/Source/AlembicLibrary/Private/AbcImporter.cpp
--- a/Engine/Plugins/Experimental/AlembicImporter/Source/AlembicLibrary/Private/AbcImporter.cpp
+++ b/Engine/Plugins/Experimental/AlembicImporter/Source/AlembicLibrary/Private/AbcImporter.cpp
@@ -52,6 +52,7 @@ THIRD_PARTY_INCLUDES_END
 #include "GeometryCacheCodecV1.h"
 #include "Subsystems/AssetEditorSubsystem.h"
 #include "Editor.h"
+#include "Animation/AnimData/AnimDataController.h"
 
 #define LOCTEXT_NAMESPACE "AbcImporter"
 
@@ -645,18 +646,22 @@ TArray<UObject*> FAbcImporter::ImportAsSkeletalMesh(UObject* InParent, EObjectFl
 		// Create animation sequence for the skeleton
 		UAnimSequence* Sequence = CreateObjectInstance<UAnimSequence>(InParent, FString::Printf(TEXT("%s_Animation"), *SkeletalMesh->GetName()), Flags);
 		Sequence->SetSkeleton(Skeleton);
-		Sequence->SetSequenceLength(AbcFile->GetImportLength());
-		Sequence->ImportFileFramerate = AbcFile->GetFramerate();
-		Sequence->ImportResampleFramerate = AbcFile->GetFramerate();
-
-		// Add 1 because of the way AnimSequence computes its framerate
-		Sequence->SetNumberOfSampledKeys(AbcFile->GetImportNumFrames() + 1);
 
 		int32 ObjectIndex = 0;
 		uint32 TriangleOffset = 0;
 		uint32 WedgeOffset = 0;
 		uint32 VertexOffset = 0;
 
+		UAnimDataController* Controller = Sequence->GetController();
+
+		Controller->OpenBracket(LOCTEXT("ImportAsSkeletalMesh", "Importing Alembic Animation"));
+
+		Controller->SetPlayLength(AbcFile->GetImportLength());
+		Controller->SetFrameRate( FFrameRate(AbcFile->GetFramerate(), 1));
+
+		Sequence->ImportFileFramerate = AbcFile->GetFramerate();
+		Sequence->ImportResampleFramerate = 1.0f / (float)AbcFile->GetFramerate();
+
 		{
 #if WITH_EDITOR
 			// When ScopedPostEditChange goes out of scope, it will call SkeletalMesh->PostEditChange()
@@ -715,7 +720,7 @@ TArray<UObject*> FAbcImporter::ImportAsSkeletalMesh(UObject* InParent, EObjectFl
 							FName ConstCurveName = *CurveName;
 
 							// Sets up the morph target curves with the sample values and time keys
-							SetupMorphTargetCurves(Skeleton, ConstCurveName, Sequence, CurveValues, TimeValues);
+							SetupMorphTargetCurves(Skeleton, ConstCurveName, Sequence, CurveValues, TimeValues, Controller);
 						}
 						else
 						{
@@ -740,10 +745,10 @@ TArray<UObject*> FAbcImporter::ImportAsSkeletalMesh(UObject* InParent, EObjectFl
 
 					const FReferenceSkeleton& RefSkeleton = SkeletalMesh->GetRefSkeleton();
 					const TArray<FMeshBoneInfo>& BonesInfo = RefSkeleton.GetRawRefBoneInfo();
-					Sequence->AddNewRawTrack(BonesInfo[0].Name, &RootBoneTrack);
-				}
 
-				Sequence->RawCurveData.RemoveRedundantKeys();
+					Controller->AddBoneTrack(BonesInfo[0].Name);
+					Controller->SetBoneTrackKeys(BonesInfo[0].Name, RootBoneTrack.PosKeys, RootBoneTrack.RotKeys, RootBoneTrack.ScaleKeys);
+				}
 
 				WedgeOffset += CompressedData.AverageSample->Indices.Num();
 				VertexOffset += CompressedData.AverageSample->Vertices.Num();
@@ -774,10 +779,10 @@ TArray<UObject*> FAbcImporter::ImportAsSkeletalMesh(UObject* InParent, EObjectFl
 
 		SkeletalMesh->MarkPackageDirty();
 
-		// Retrieve the name mapping container
-		const FSmartNameMapping* NameMapping = Skeleton->GetSmartNameContainer(USkeleton::AnimCurveMappingName);
-		Sequence->RawCurveData.RefreshName(NameMapping);
-		Sequence->MarkRawDataAsModified();
+		Controller->UpdateCurveNamesFromSkeleton(Skeleton, ERawCurveTrackTypes::RCT_Float);
+
+		Controller->CloseBracket();
+
 		Sequence->PostEditChange();
 		Sequence->SetPreviewMesh(SkeletalMesh);
 		Sequence->MarkPackageDirty();
@@ -803,20 +808,24 @@ TArray<UObject*> FAbcImporter::ImportAsSkeletalMesh(UObject* InParent, EObjectFl
 	return GeneratedObjects;
 }
 
-void FAbcImporter::SetupMorphTargetCurves(USkeleton* Skeleton, FName ConstCurveName, UAnimSequence* Sequence, const TArray<float> &CurveValues, const TArray<float>& TimeValues)
+void FAbcImporter::SetupMorphTargetCurves(USkeleton* Skeleton, FName ConstCurveName, UAnimSequence* Sequence, const TArray<float> &CurveValues, const TArray<float>& TimeValues, UAnimDataController* Controller)
 {
 	FSmartName NewName;
 	Skeleton->AddSmartNameAndModify(USkeleton::AnimCurveMappingName, ConstCurveName, NewName);
 
-	check(Sequence->RawCurveData.AddCurveData(NewName));
-	FFloatCurve * NewCurve = static_cast<FFloatCurve *> (Sequence->RawCurveData.GetCurveData(NewName.UID, ERawCurveTrackTypes::RCT_Float));
+	FAnimationCurveIdentifier CurveId(NewName, ERawCurveTrackTypes::RCT_Float);
+	Controller->AddCurve(CurveId);
 
+	const FFloatCurve* NewCurve = Sequence->GetDataModel()->FindFloatCurve(CurveId);
+	ensure(NewCurve);
+
+	FRichCurve RichCurve;
 	for (int32 KeyIndex = 0; KeyIndex < CurveValues.Num(); ++KeyIndex)
 	{
 		const float CurveValue = CurveValues[KeyIndex];
 		const float TimeValue = TimeValues[KeyIndex];
 
-		FKeyHandle NewKeyHandle = NewCurve->FloatCurve.AddKey(TimeValue, CurveValue, false);
+		FKeyHandle NewKeyHandle = RichCurve.AddKey(TimeValue, CurveValue, false);
 
 		ERichCurveInterpMode NewInterpMode = RCIM_Linear;
 		ERichCurveTangentMode NewTangentMode = RCTM_Auto;
@@ -827,9 +836,10 @@ void FAbcImporter::SetupMorphTargetCurves(USkeleton* Skeleton, FName ConstCurveN
 		float LeaveTangentWeight = 0.f;
 		float ArriveTangentWeight = 0.f;
 
-		NewCurve->FloatCurve.SetKeyInterpMode(NewKeyHandle, NewInterpMode);
-		NewCurve->FloatCurve.SetKeyTangentMode(NewKeyHandle, NewTangentMode);
-		NewCurve->FloatCurve.SetKeyTangentWeightMode(NewKeyHandle, NewTangentWeightMode);
+		RichCurve.SetKeyInterpMode(NewKeyHandle, NewInterpMode);
+		RichCurve.SetKeyTangentMode(NewKeyHandle, NewTangentMode);
+		RichCurve.SetKeyTangentWeightMode(NewKeyHandle, NewTangentWeightMode);
+
 
 		// When data isn't compressed, add keys to remove inbetween frames interpolation
 		if (ImportSettings && ImportSettings->CompressionSettings.BaseCalculationType == EBaseCalculationType::NoCompression)
@@ -843,15 +853,17 @@ void FAbcImporter::SetupMorphTargetCurves(USkeleton* Skeleton, FName ConstCurveN
 				{
 					const float NextTimeValue = TimeValues[NextKeyIndex];
 					const float SubTimeStep = (NextTimeValue - TimeValue) / 20.f;
-					FKeyHandle StepKeyHandle = NewCurve->FloatCurve.AddKey(NextTimeValue - SubTimeStep, CurveValue, false);
+					FKeyHandle StepKeyHandle = RichCurve.AddKey(NextTimeValue - SubTimeStep, CurveValue, false);
 
-					NewCurve->FloatCurve.SetKeyInterpMode(StepKeyHandle, NewInterpMode);
-					NewCurve->FloatCurve.SetKeyTangentMode(StepKeyHandle, NewTangentMode);
-					NewCurve->FloatCurve.SetKeyTangentWeightMode(StepKeyHandle, NewTangentWeightMode);
+					RichCurve.SetKeyInterpMode(StepKeyHandle, NewInterpMode);
+					RichCurve.SetKeyTangentMode(StepKeyHandle, NewTangentMode);
+					RichCurve.SetKeyTangentWeightMode(StepKeyHandle, NewTangentWeightMode);
 				}
 			}
 		}
 	}
+
+	Controller->SetCurveKeys(CurveId, RichCurve.GetConstRefOfKeys());
 }
 
 const bool FAbcImporter::CompressAnimationDataUsingPCA(const FAbcCompressionSettings& InCompressionSettings, const bool bRunComparison /*= false*/)
diff --git a/Engine/Plugins/Experimental/AlembicImporter/Source/AlembicLibrary/Public/AbcImporter.h b/Engine/Plugins/Experimental/AlembicImporter/Source/AlembicLibrary/Public/AbcImporter.h
--- a/Engine/Plugins/Experimental/AlembicImporter/Source/AlembicLibrary/Public/AbcImporter.h
+++ b/Engine/Plugins/Experimental/AlembicImporter/Source/AlembicLibrary/Public/AbcImporter.h
@@ -32,6 +32,8 @@ class UAbcImportSettings;
 class FSkeletalMeshImportData;
 class UAbcAssetImportData;
 
+class UAnimDataController;
+
 struct FMeshDescription;
 struct FAbcImportData;
 struct FGeometryCacheMeshData;
@@ -200,7 +202,7 @@ private:
 	void GenerateMorphTargetVertices(FAbcMeshSample* BaseSample, TArray<FMorphTargetDelta> &MorphDeltas, FAbcMeshSample* AverageSample, uint32 WedgeOffset, const TArray<int32>& RemapIndices, const TArray<int32>& UsedVertexIndicesForMorphs, const uint32 VertexOffset, const uint32 IndexOffset, bool bEnableDeltaOffset, FVector& OutOffset);
 	
 	/** Set up correct morph target weights from the PCA compressed data */
-	void SetupMorphTargetCurves(USkeleton* Skeleton, FName ConstCurveName, UAnimSequence* Sequence, const TArray<float> &CurveValues, const TArray<float>& TimeValues);
+	void SetupMorphTargetCurves(USkeleton* Skeleton, FName ConstCurveName, UAnimSequence* Sequence, const TArray<float> &CurveValues, const TArray<float>& TimeValues, UAnimDataController* Controller);
 	
 private:
 	/** Cached ptr for the import settings */
diff --git a/Engine/Plugins/Experimental/ControlRig/Source/ControlRig/Private/Sequencer/MovieSceneControlRigParameterSection.cpp b/Engine/Plugins/Experimental/ControlRig/Source/ControlRig/Private/Sequencer/MovieSceneControlRigParameterSection.cpp
--- a/Engine/Plugins/Experimental/ControlRig/Source/ControlRig/Private/Sequencer/MovieSceneControlRigParameterSection.cpp
+++ b/Engine/Plugins/Experimental/ControlRig/Source/ControlRig/Private/Sequencer/MovieSceneControlRigParameterSection.cpp
@@ -12,6 +12,7 @@
 #include "Animation/AnimSequence.h"
 #include "ControlRig/Private/Units/Execution/RigUnit_InverseExecution.h"
 #include "Misc/ScopedSlowTask.h"
+#include "Animation/AnimSequenceHelpers.h"
 
 #define LOCTEXT_NAMESPACE "MovieSceneControlParameterRigSection"
 
@@ -2097,28 +2098,29 @@ bool UMovieSceneControlRigParameterSection::LoadAnimSequenceIntoThisSection(UAni
 	FScopedSlowTask Progress(NumberOfFrames + ExtraProgress, LOCTEXT("BakingToControlRig_SlowTask", "Baking To Control Rig..."));	
 	Progress.MakeDialog(true);
 
+	const UAnimDataModel* DataModel = AnimSequence->GetDataModel();
+	const FAnimationCurveData& CurveData = DataModel->GetCurveData();
+	const TArray<FBoneAnimationTrack>& BoneAnimationTracks = DataModel->GetBoneAnimationTracks();
+
 	for (int32 Index = 0; Index < NumberOfFrames; ++Index)
 	{
-		const float SequenceSecond = FrameRate.AsSeconds(Index);
+		const float SequenceSecond = AnimSequence->GetTimeAtFrame(Index);
 		FFrameNumber FrameNumber = StartFrame + (FrameRateInFrameNumber * Index);
 
-		for (FFloatCurve& Curve : AnimSequence->RawCurveData.FloatCurves)
+		for (const FFloatCurve& Curve : CurveData.FloatCurves)
 		{
-			float Val = Curve.FloatCurve.Eval(SequenceSecond);
+			const float Val = Curve.FloatCurve.Eval(SequenceSecond);
 			SourceCurves.SetValue(Curve.Name.DisplayName,Val);
 		}
-		for (int32 TrackIndex = 0; TrackIndex < AnimSequence->GetRawAnimationData().Num(); ++TrackIndex)
+
+		for(int32 TrackIndex = 0; TrackIndex < BoneAnimationTracks.Num(); ++TrackIndex)
 		{
-			// verify if this bone exists in skeleton
-			int32 BoneTreeIndex = AnimSequence->GetSkeletonIndexFromRawDataTrackIndex(TrackIndex);
-			if (BoneTreeIndex != INDEX_NONE)
-			{
-				FName BoneName = Skeleton->GetReferenceSkeleton().GetBoneName(BoneTreeIndex);
-				FTransform BoneTransform;
-				AnimSequence->GetRawAnimationTrack(TrackIndex);
-				AnimSequence->ExtractBoneTransform(AnimSequence->GetRawAnimationTrack(TrackIndex), BoneTransform, SequenceSecond);
-				SourceBones.SetLocalTransform(BoneName, BoneTransform);
-			}
+			const FBoneAnimationTrack& AnimationTrack = BoneAnimationTracks[TrackIndex];
+
+			FName BoneName = Skeleton->GetReferenceSkeleton().GetBoneName(AnimationTrack.BoneTreeIndex);
+			FTransform BoneTransform;
+			UE::Anim::GetBoneTransformFromModel(DataModel, BoneTransform, TrackIndex, SequenceSecond, EAnimInterpolationType::Linear);
+			SourceBones.SetLocalTransform(BoneName, BoneTransform);
 		}
 
 		ControlRig->Execute(EControlRigState::Update, FRigUnit_InverseExecution::EventName);
diff --git a/Engine/Plugins/Experimental/ControlRig/Source/ControlRigEditor/Private/Sequencer/ControlRigParameterTrackEditor.cpp b/Engine/Plugins/Experimental/ControlRig/Source/ControlRigEditor/Private/Sequencer/ControlRigParameterTrackEditor.cpp
--- a/Engine/Plugins/Experimental/ControlRig/Source/ControlRigEditor/Private/Sequencer/ControlRigParameterTrackEditor.cpp
+++ b/Engine/Plugins/Experimental/ControlRig/Source/ControlRigEditor/Private/Sequencer/ControlRigParameterTrackEditor.cpp
@@ -2853,7 +2853,7 @@ void FControlRigParameterSection::OnAnimationAssetSelectedForFK(const FAssetData
 		UObject* BoundObject = nullptr;
 		USkeleton* Skeleton = AcquireSkeletonFromObjectGuid(ObjectBinding, &BoundObject, SequencerPtr);
 
-		if (AnimSequence && AutoRig && Skeleton && AnimSequence->GetRawAnimationData().Num() > 0)
+		if (AnimSequence && AutoRig && Skeleton && AnimSequence->GetDataModel()->GetNumBoneTracks() > 0)
 		{
 			FScopedTransaction Transaction(LOCTEXT("BakeAnimation_Transaction", "Bake Animation To FK Control Rig"));
 			Section->Modify();
diff --git a/Engine/Plugins/Experimental/MotionTrailEditorMode/Source/MotionTrailEditorMode/Private/Sequencer/AnimationBoneTrail.cpp b/Engine/Plugins/Experimental/MotionTrailEditorMode/Source/MotionTrailEditorMode/Private/Sequencer/AnimationBoneTrail.cpp
--- a/Engine/Plugins/Experimental/MotionTrailEditorMode/Source/MotionTrailEditorMode/Private/Sequencer/AnimationBoneTrail.cpp
+++ b/Engine/Plugins/Experimental/MotionTrailEditorMode/Source/MotionTrailEditorMode/Private/Sequencer/AnimationBoneTrail.cpp
@@ -12,6 +12,7 @@
 #include "ISequencer.h"
 #include "MovieSceneSequence.h"
 #include "Exporters/AnimSeqExportOption.h"
+#include "Animation/AnimSequenceHelpers.h"
 
 namespace UE 
 {
@@ -111,7 +112,7 @@ void FAnimTrajectoryCache::GetSpaceBasedAnimationData(TArray<TArray<FTransform>>
 		for (int32 BoneIndex = 0; BoneIndex < NumBones; ++BoneIndex)
 		{
 			FName BoneName = MySkeleton->GetReferenceSkeleton().GetBoneName(BoneIndex);
-			int32 TrackIndex = CachedAnimSequence->GetAnimationTrackNames().Find(BoneName);
+			int32 TrackIndex = CachedAnimSequence->GetDataModel()->GetBoneTrackIndexByName(BoneName);
 			SkelToTrackIdx[BoneIndex] = TrackIndex;
 		}
 	}
@@ -124,12 +125,11 @@ void FAnimTrajectoryCache::GetSpaceBasedAnimationData(TArray<TArray<FTransform>>
 
 		if (TrackIndex != INDEX_NONE)
 		{
-			auto& RawAnimation = CachedAnimSequence->GetRawAnimationData()[TrackIndex];
 			// fill up keys - calculate PK1 * K1
 			for (int32 Key = 0; Key < NumKeys; ++Key)
 			{
 				FTransform AnimatedLocalKey;
-				CachedAnimSequence->ExtractBoneTransform(CachedAnimSequence->GetRawAnimationData(), AnimatedLocalKey, TrackIndex, Interval*Key);
+				UE::Anim::GetBoneTransformFromModel(CachedAnimSequence->GetDataModel(), AnimatedLocalKey, TrackIndex, Interval * Key, CachedAnimSequence->Interpolation);
 
 				if (ParentBoneIndex != INDEX_NONE)
 				{
diff --git a/Engine/Plugins/Importers/USDImporter/Source/USDSchemas/Private/USDSkelRootTranslator.cpp b/Engine/Plugins/Importers/USDImporter/Source/USDSchemas/Private/USDSkelRootTranslator.cpp
--- a/Engine/Plugins/Importers/USDImporter/Source/USDSchemas/Private/USDSkelRootTranslator.cpp
+++ b/Engine/Plugins/Importers/USDImporter/Source/USDSchemas/Private/USDSkelRootTranslator.cpp
@@ -864,7 +864,7 @@ namespace UsdSkelRootTranslatorImpl
 							TUsdStore<pxr::VtArray<pxr::UsdSkelSkinningQuery>> SkinningTargets = Binding.GetSkinningTargets();
 							UsdToUnreal::ConvertSkelAnim( SkelQuery, &SkinningTargets.Get(), &NewBlendShapes, Context->bAllowInterpretingLODs, AnimSequence );
 
-							if ( AnimSequence->GetRawAnimationData().Num() != 0 || AnimSequence->RawCurveData.FloatCurves.Num() != 0 )
+							if (AnimSequence->GetDataModel()->GetNumBoneTracks() != 0 || AnimSequence->GetDataModel()->GetNumberOfFloatCurves() != 0 )
 							{
 								UUsdAssetImportData* ImportData = NewObject< UUsdAssetImportData >( AnimSequence, TEXT( "USDAssetImportData" ) );
 								ImportData->PrimPath = GetPrim().GetPrimPath().GetString(); // Point to the SkelRoot so that it ends up next to the skeletal mesh
diff --git a/Engine/Plugins/Importers/USDImporter/Source/USDUtilities/Private/USDSkeletalDataConversion.cpp b/Engine/Plugins/Importers/USDImporter/Source/USDUtilities/Private/USDSkeletalDataConversion.cpp
--- a/Engine/Plugins/Importers/USDImporter/Source/USDUtilities/Private/USDSkeletalDataConversion.cpp
+++ b/Engine/Plugins/Importers/USDImporter/Source/USDUtilities/Private/USDSkeletalDataConversion.cpp
@@ -408,14 +408,17 @@ namespace SkelDataConversionImpl
 		FSmartName NewName;
 		Skeleton->AddSmartNameAndModify( USkeleton::AnimCurveMappingName, CurveName, NewName );
 
-		FFloatCurve* Curve = static_cast< FFloatCurve* >( Sequence->RawCurveData.GetCurveData( NewName.UID, ERawCurveTrackTypes::RCT_Float ) );
+
+		const UAnimDataModel* DataModel = Sequence->GetDataModel();
+		UAnimDataController* Controller = Sequence->GetController();
+
+		FAnimationCurveIdentifier CurveId(NewName, ERawCurveTrackTypes::RCT_Float);
+		const FFloatCurve* Curve = DataModel->FindFloatCurve(CurveId);
 		if ( !Curve )
 		{
-			if ( Sequence->RawCurveData.AddCurveData( NewName, AACF_DefaultCurve ) )
-			{
-				Curve = static_cast< FFloatCurve* > ( Sequence->RawCurveData.GetCurveData( NewName.UID, ERawCurveTrackTypes::RCT_Float ) );
-				Curve->Name = NewName;
-			}
+			// If curve doesn't exist, add one
+			Controller->AddCurve(CurveId, AACF_DefaultCurve);
+			Curve = DataModel->FindFloatCurve(CurveId);
 		}
 		else
 		{
@@ -428,16 +431,14 @@ namespace SkelDataConversionImpl
 				);
 			}
 
-			Curve->FloatCurve.Reset();
-			Curve->SetCurveTypeFlags( Curve->GetCurveTypeFlags() | AACF_DefaultCurve );
+			Controller->SetCurveFlags(CurveId, Curve->GetCurveTypeFlags() | AACF_DefaultCurve);
 		}
 
-		Sequence->RawCurveData.RefreshName( NameMapping );
+		Controller->UpdateCurveNamesFromSkeleton(Skeleton, ERawCurveTrackTypes::RCT_Float);
 
 		if ( Curve )
 		{
-			Curve->FloatCurve = SourceData;
-			Curve->FloatCurve.RemoveRedundantKeys( KINDA_SMALL_NUMBER );
+			Controller->SetCurveKeys(CurveId, SourceData.GetConstRefOfKeys());
 		}
 		else
 		{
@@ -1633,8 +1634,6 @@ bool UsdToUnreal::ConvertSkelAnim( const pxr::UsdSkelSkeletonQuery& InUsdSkeleto
 		return false;
 	}
 
-	OutSkeletalAnimationAsset->CleanAnimSequenceForImport();
-
 	TUsdStore<std::vector<double>> UsdJointTransformTimeSamples;
 	AnimQuery.Get().GetJointTransformTimeSamples( &( UsdJointTransformTimeSamples.Get() ) );
 	int32 NumJointTransformSamples = UsdJointTransformTimeSamples.Get().size();
@@ -1667,6 +1666,14 @@ bool UsdToUnreal::ConvertSkelAnim( const pxr::UsdSkelSkeletonQuery& InUsdSkeleto
 	const int32 NumBakedFrames = FMath::CeilToInt( FMath::Max( SequenceLengthSeconds * FramesPerSecond + 1.0, 1.0 ) );
 	const double IntervalTimeCodes = ( NumBakedFrames > 1 ) ? ( SequenceLengthTimeCodes / ( NumBakedFrames - 1 ) ) : MINIMUM_ANIMATION_LENGTH;
 
+
+	const UAnimDataModel* DataModel = OutSkeletalAnimationAsset->GetDataModel();
+	UAnimDataController* Controller = OutSkeletalAnimationAsset->GetController();
+
+	Controller->OpenBracket(LOCTEXT("ImportUSDAnimData_Bracket", "Importing USD Animation Data"));
+
+	Controller->ResetModel();
+
 	// Bake the animation for each frame.
 	// An alternative route would be to convert the time samples into TransformCurves, add them to UAnimSequence::RawCurveData,
 	// and then call UAnimSequence::BakeTrackCurvesToRawAnimation. Doing it this way provides a few benefits though: The main one is that the way with which
@@ -1708,7 +1715,8 @@ bool UsdToUnreal::ConvertSkelAnim( const pxr::UsdSkelSkeletonQuery& InUsdSkeleto
 
 		for ( int32 BoneIndex = 0; BoneIndex < NumBones; ++BoneIndex )
 		{
-			OutSkeletalAnimationAsset->AddNewRawTrack( BoneInfo[ BoneIndex ].Name, &JointTracks[ BoneIndex ] );
+			Controller->AddBoneTrack(BoneInfo[BoneIndex].Name); 
+			Controller->SetBoneTrackKeys(BoneInfo[BoneIndex].Name, JointTracks[BoneIndex].PosKeys, JointTracks[BoneIndex].RotKeys, JointTracks[BoneIndex].ScaleKeys);
 		}
 	}
 
@@ -1894,10 +1902,10 @@ bool UsdToUnreal::ConvertSkelAnim( const pxr::UsdSkelSkeletonQuery& InUsdSkeleto
 	OutSkeletalAnimationAsset->Interpolation = Stage.Get()->GetInterpolationType() == pxr::UsdInterpolationTypeHeld ? EAnimInterpolationType::Step : EAnimInterpolationType::Linear;
 	OutSkeletalAnimationAsset->ImportFileFramerate = Stage.Get()->GetFramesPerSecond();
 	OutSkeletalAnimationAsset->ImportResampleFramerate = FramesPerSecond;
-	OutSkeletalAnimationAsset->SetSequenceLength(SequenceLengthSeconds);
-	OutSkeletalAnimationAsset->SetNumberOfSampledKeys( NumBakedFrames );
-	OutSkeletalAnimationAsset->MarkRawDataAsModified();	
-	OutSkeletalAnimationAsset->PostProcessSequence();
+
+	Controller->SetPlayLength(SequenceLengthSeconds);
+	Controller->SetFrameRate(FFrameRate(FramesPerSecond, 1));
+	Controller->CloseBracket();
 
 	OutSkeletalAnimationAsset->PostEditChange();
 	OutSkeletalAnimationAsset->MarkPackageDirty();
diff --git a/Engine/Plugins/VirtualProduction/Takes/Source/TakeTrackRecorders/Private/MovieScene3DTransformTrackRecorder.cpp b/Engine/Plugins/VirtualProduction/Takes/Source/TakeTrackRecorders/Private/MovieScene3DTransformTrackRecorder.cpp
--- a/Engine/Plugins/VirtualProduction/Takes/Source/TakeTrackRecorders/Private/MovieScene3DTransformTrackRecorder.cpp
+++ b/Engine/Plugins/VirtualProduction/Takes/Source/TakeTrackRecorders/Private/MovieScene3DTransformTrackRecorder.cpp
@@ -12,6 +12,7 @@
 #include "Components/SceneComponent.h"
 #include "Components/SkeletalMeshComponent.h"
 #include "SequenceRecorderUtils.h"
+#include "Animation/AnimData/AnimDataModel.h"
 
 
 DEFINE_LOG_CATEGORY(TransformSerialization);
@@ -495,14 +496,16 @@ void UMovieScene3DTransformTrackRecorder::PostProcessAnimationData(UMovieSceneAn
 				// Find the root bone
 				int32 RootIndex = INDEX_NONE;
 				USkeleton* AnimSkeleton = AnimSequence->GetSkeleton();
-				for (int32 TrackIndex = 0; TrackIndex < AnimSequence->GetRawAnimationData().Num(); ++TrackIndex)
+
+				const TArray<FBoneAnimationTrack>& BoneAnimationTracks = AnimSequence->GetDataModel()->GetBoneAnimationTracks();
+				for (const FBoneAnimationTrack& AnimationTrack : BoneAnimationTracks)
 				{
-					// Vverify that this bone exists in skeleton
-					int32 BoneTreeIndex = AnimSequence->GetSkeletonIndexFromRawDataTrackIndex(TrackIndex);
+					// Verify if this bone exists in skeleton
+					const int32 BoneTreeIndex = AnimationTrack.BoneTreeIndex;
 					if (BoneTreeIndex != INDEX_NONE)
 					{
-						int32 BoneIndex = AnimSkeleton->GetMeshBoneIndexFromSkeletonBoneIndex(SkeletalMesh, BoneTreeIndex);
-						int32 ParentIndex = SkeletalMesh->GetRefSkeleton().GetParentIndex(BoneIndex);
+						const int32 BoneIndex = AnimSkeleton->GetMeshBoneIndexFromSkeletonBoneIndex(SkeletalMesh, BoneTreeIndex);
+						const int32 ParentIndex = SkeletalMesh->GetRefSkeleton().GetParentIndex(BoneIndex);
 						if (ParentIndex == INDEX_NONE)
 						{
 							// We've found the root (root bones do not have a valid parent)
@@ -519,12 +522,14 @@ void UMovieScene3DTransformTrackRecorder::PostProcessAnimationData(UMovieSceneAn
 					return;
 				}
 
+				ensure(BoneAnimationTracks.IsValidIndex(RootIndex));
+
 				const FFrameRate TickResolution = MovieSceneSection->GetTypedOuter<UMovieScene>()->GetTickResolution();
 				const FFrameNumber StartTime = MovieSceneSection->GetInclusiveStartFrame();
 
 				FTransform InvComponentTransform = AnimTrackRecorder->GetComponentTransform().Inverse();
 
-				const FRawAnimSequenceTrack& RawTrack = AnimSequence->GetRawAnimationData()[RootIndex];
+				const FRawAnimSequenceTrack& RawTrack = BoneAnimationTracks[RootIndex].InternalTrackData;
 				const int32 KeyCount = FMath::Max(FMath::Max(RawTrack.PosKeys.Num(), RawTrack.RotKeys.Num()), RawTrack.ScaleKeys.Num());
 				for (int32 KeyIndex = 0; KeyIndex < KeyCount; KeyIndex++)
 				{
diff --git a/Engine/Plugins/VirtualProduction/Takes/Source/TakeTrackRecorders/Private/MovieSceneAnimationTrackRecorder.cpp b/Engine/Plugins/VirtualProduction/Takes/Source/TakeTrackRecorders/Private/MovieSceneAnimationTrackRecorder.cpp
--- a/Engine/Plugins/VirtualProduction/Takes/Source/TakeTrackRecorders/Private/MovieSceneAnimationTrackRecorder.cpp
+++ b/Engine/Plugins/VirtualProduction/Takes/Source/TakeTrackRecorders/Private/MovieSceneAnimationTrackRecorder.cpp
@@ -15,6 +15,9 @@
 #include "Engine/Engine.h"
 #include "LevelSequence.h"
 #include "UObject/UObjectBaseUtility.h"
+#include "Animation/AnimData/AnimDataController.h"
+
+#define LOCTEXT_NAMESPACE "MovieSceneAnimationTrackRecorder"
 
 DEFINE_LOG_CATEGORY(AnimationSerialization);
 
@@ -362,32 +365,47 @@ bool UMovieSceneAnimationTrackRecorder::LoadRecordedFile(const FString& FileName
 
 								CreateAnimationAssetAndSequence(Actor,AnimationDirectory);
 
-								AnimSequence->RecycleAnimSequence();
-								AnimSequence->SetSequenceLength(0.f);
-								AnimSequence->SetNumberOfSampledKeys(0);
-								for (int32 TrackIndex = 0; TrackIndex < Header.AnimationTrackNames.Num(); ++TrackIndex)
+								AnimSequence->DeleteNotifyTrackData();
+
+								UAnimDataController* Controller = AnimSequence->GetController();
 								{
-									AnimSequence->AddNewRawTrack(Header.AnimationTrackNames[TrackIndex]);
+									UAnimDataController::FScopedBracket ScopedBracket(Controller, LOCTEXT("LoadRecordedFile_Bracket", "Loading recorded animation file"));
 
-								}
-								AnimSequence->InitializeNotifyTrack();
+									Controller->ResetModel();
 
-								for (const FAnimationSerializedFrame& SerializedFrame : InFrames)
-								{
-									const FSerializedAnimation &Frame = SerializedFrame.Frame;
-									for (int32 TrackIndex = 0; TrackIndex < Frame.AnimationData.Num(); ++TrackIndex)
+									const float FloatDenominator = 1000.0f;
+									const float Numerator = FloatDenominator / Header.IntervalTime;
+									Controller->SetFrameRate(FFrameRate(Numerator, FloatDenominator));
+
+									int32 MaxNumberOfKeys = 0;
+									for (int32 TrackIndex = 0; TrackIndex < Header.AnimationTrackNames.Num(); ++TrackIndex)
 									{
-										FRawAnimSequenceTrack& RawTrack = AnimSequence->GetRawAnimationTrack(TrackIndex);
-										RawTrack.PosKeys.Add(Frame.AnimationData[TrackIndex].PosKey);
-										RawTrack.RotKeys.Add(Frame.AnimationData[TrackIndex].RotKey);
-										RawTrack.ScaleKeys.Add(Frame.AnimationData[TrackIndex].ScaleKey);
+										Controller->AddBoneTrack(Header.AnimationTrackNames[TrackIndex]);
+
+										TArray<FVector> PosKeys;
+										TArray<FQuat> RotKeys;
+										TArray<FVector> ScaleKeys;
+
+										// Generate key arrays
+										for (const FAnimationSerializedFrame& SerializedFrame : InFrames)
+										{
+											const FSerializedAnimation& Frame = SerializedFrame.Frame;
+											PosKeys.Add(Frame.AnimationData[TrackIndex].PosKey);
+											RotKeys.Add(Frame.AnimationData[TrackIndex].RotKey);
+											ScaleKeys.Add(Frame.AnimationData[TrackIndex].ScaleKey);
+										}
+
+										MaxNumberOfKeys = FMath::Max(MaxNumberOfKeys, PosKeys.Num());
+										MaxNumberOfKeys = FMath::Max(MaxNumberOfKeys, RotKeys.Num());
+										MaxNumberOfKeys = FMath::Max(MaxNumberOfKeys, ScaleKeys.Num());
+
+										Controller->SetBoneTrackKeys(Header.AnimationTrackNames[TrackIndex], PosKeys, RotKeys, ScaleKeys);
 									}
+
+									Controller->SetPlayLength((MaxNumberOfKeys > 1) ? (MaxNumberOfKeys - 1) * Header.IntervalTime : MINIMUM_ANIMATION_LENGTH);
+									Controller->NotifyPopulated();
 								}
-								AnimSequence->SetNumberOfSampledKeys( InFrames.Num() );
-								AnimSequence->SetSequenceLength((AnimSequence->GetNumberOfSampledKeys() > 1) ? (AnimSequence->GetNumberOfSampledKeys() - 1) * Header.IntervalTime : MINIMUM_ANIMATION_LENGTH);
 
-								//fix up notifies todo notifies mz
-								AnimSequence->PostProcessSequence();
 								AnimSequence->MarkPackageDirty();
 								FFrameRate TickResolution = InMovieScene->GetTickResolution();;
 
@@ -424,4 +442,5 @@ bool UMovieSceneAnimationTrackRecorder::LoadRecordedFile(const FString& FileName
 	}
 	
 	return false;
-}
\ No newline at end of file
+}
+#undef LOCTEXT_NAMESPACE // "MovieSceneAnimationTrackRecorder"
\ No newline at end of file
diff --git a/Engine/Source/Developer/AssetTools/Private/AssetTypeActions/AssetTypeActions_Skeleton.cpp b/Engine/Source/Developer/AssetTools/Private/AssetTypeActions/AssetTypeActions_Skeleton.cpp
--- a/Engine/Source/Developer/AssetTools/Private/AssetTypeActions/AssetTypeActions_Skeleton.cpp
+++ b/Engine/Source/Developer/AssetTools/Private/AssetTypeActions/AssetTypeActions_Skeleton.cpp
@@ -903,22 +903,26 @@ void FAssetTypeActions_Skeleton::RetargetSkeleton(TArray<FAssetToRemapSkeleton>&
 			UObject* Asset = RenameData.Asset.Get();
 			if ( Asset )
 			{
-				if ( UAnimationAsset * AnimAsset = Cast<UAnimationAsset>(Asset) )
+				if ( UAnimationAsset* AnimAsset = Cast<UAnimationAsset>(Asset) )
 				{
-					UAnimSequenceBase * SequenceBase = Cast<UAnimSequenceBase>(AnimAsset);
-					if (SequenceBase)
+					TUniquePtr<UAnimDataController::FScopedBracket> ScopedBracket;
+
+					if (UAnimSequenceBase* SequenceBase = Cast<UAnimSequenceBase>(AnimAsset))
 					{
-						EditorAnimUtils::CopyAnimCurves(OldSkeleton, NewSkeleton, SequenceBase, USkeleton::AnimCurveMappingName, ERawCurveTrackTypes::RCT_Float);
+						UAnimDataController* Controller = SequenceBase->GetController();
+						ScopedBracket = MakeUnique<UAnimDataController::FScopedBracket>(Controller, LOCTEXT("FAssetTypeActions_Skeleton_Bracket", "Retargeting Skeleton"));
+
+						Controller->FindOrAddCurveNamesOnSkeleton(NewSkeleton, ERawCurveTrackTypes::RCT_Float);
 						
-						if (UAnimSequence * Sequence = Cast<UAnimSequence>(SequenceBase))
+						if (UAnimSequence* Sequence = Cast<UAnimSequence>(SequenceBase))
 						{
-							EditorAnimUtils::CopyAnimCurves(OldSkeleton, NewSkeleton, Sequence, USkeleton::AnimTrackCurveMappingName, ERawCurveTrackTypes::RCT_Transform);
+							Controller->FindOrAddCurveNamesOnSkeleton(NewSkeleton, ERawCurveTrackTypes::RCT_Transform);
 						}
 					}
 					
 					AnimAsset->ReplaceSkeleton(NewSkeleton, bConvertSpaces);
 				}
-				else if ( UAnimBlueprint * AnimBlueprint = Cast<UAnimBlueprint>(Asset) )
+				else if ( UAnimBlueprint* AnimBlueprint = Cast<UAnimBlueprint>(Asset) )
 				{
 					AnimBlueprints.Add(AnimBlueprint);
 				}
diff --git a/Engine/Source/Editor/AnimGraph/Private/AnimPreviewInstance.cpp b/Engine/Source/Editor/AnimGraph/Private/AnimPreviewInstance.cpp
--- a/Engine/Source/Editor/AnimGraph/Private/AnimPreviewInstance.cpp
+++ b/Engine/Source/Editor/AnimGraph/Private/AnimPreviewInstance.cpp
@@ -5,6 +5,7 @@
 #include "Animation/DebugSkelMeshComponent.h"
 #include "AnimationRuntime.h"
 #include "Animation/AnimSequence.h"
+#include "Animation/AnimSequenceHelpers.h"
 
 #if WITH_EDITOR
 #include "ScopedTransaction.h"
@@ -268,14 +269,14 @@ void FAnimPreviewInstanceProxy::RefreshCurveBoneControllers(UAnimationAsset* Ass
 	if(CurrentSequence)
 	{
 		// make sure if this needs source update
-		if ( !CurrentSequence->DoesContainTransformCurves() )
+		if (CurrentSequence->GetDataModel()->GetNumberOfTransformCurves() == 0)
 		{
 			return;
 		}
 
 		GetRequiredBones().SetUseSourceData(true);
 
-		TArray<FTransformCurve>& Curves = CurrentSequence->RawCurveData.TransformCurves;
+		const TArray<FTransformCurve>& Curves = CurrentSequence->GetDataModel()->GetCurveData().TransformCurves;
 		USkeleton* MySkeleton = CurrentSequence->GetSkeleton();
 		for (auto& Curve : Curves)
 		{
@@ -303,7 +304,7 @@ void FAnimPreviewInstanceProxy::UpdateCurveController()
 	if (CurrentSequence && PreviewSkeleton)
 	{
 		TMap<FName, FTransform> ActiveCurves;
-		CurrentSequence->RawCurveData.EvaluateTransformCurveData(PreviewSkeleton, ActiveCurves, GetCurrentTime(), 1.f);
+		UE::Anim::EvaluateTransformCurvesFromModel(CurrentSequence->GetDataModel(), ActiveCurves, GetCurrentTime(), 1.f);
 
 		// make sure those curves exists in the bone controller, otherwise problem
 		if ( ActiveCurves.Num() > 0 )
diff --git a/Engine/Source/Editor/AnimationEditor/Private/AnimationEditor.cpp b/Engine/Source/Editor/AnimationEditor/Private/AnimationEditor.cpp
--- a/Engine/Source/Editor/AnimationEditor/Private/AnimationEditor.cpp
+++ b/Engine/Source/Editor/AnimationEditor/Private/AnimationEditor.cpp
@@ -685,13 +685,14 @@ void FAnimationEditor::OnRemoveBoneTrack()
 {
 	if ( FMessageDialog::Open(EAppMsgType::YesNo, LOCTEXT("WarningOnRemovingBoneTracks", "This will clear all bone transform of the animation, source data, and edited layer information. This doesn't remove notifies, and curves. Do you want to continue?")) == EAppReturnType::Yes)
 	{
-		FScopedTransaction ScopedTransaction(LOCTEXT("RemoveAnimation", "Remove Track"));
-
 		UAnimSequence* AnimSequence = Cast<UAnimSequence>(AnimationAsset);
 		if (AnimSequence)
 		{
-			AnimSequence->Modify();
-			AnimSequence->RemoveAllTracks();
+			UAnimDataController* Controller = AnimSequence->GetController();
+			UAnimDataController::FScopedBracket ScopedBracket(Controller, LOCTEXT("OnRemoveBoneTrack_Bracket", "Removing all Bone Animation and Transform Curve Tracks"));
+
+			Controller->RemoveAllBoneTracks();
+			Controller->RemoveAllCurvesOfType(ERawCurveTrackTypes::RCT_Transform);			
 		}
 	}
 }
@@ -773,13 +774,10 @@ void FAnimationEditor::CopyCurveToSoundWave(const FAssetData& SoundWaveAssetData
 	UCurveTable* CurveTable = SoundWave->GetInternalCurveData();
 
 	// iterate over curves in anim data
-	const int32 NumCurves = Sequence->RawCurveData.FloatCurves.Num();
-	for (int32 CurveIdx = 0; CurveIdx < NumCurves; CurveIdx++)
+	for (const FFloatCurve& FloatCurve : Sequence->GetDataModel()->GetFloatCurves())
 	{
-		FFloatCurve& AnimCurve = Sequence->RawCurveData.FloatCurves[CurveIdx];
-
-		FRichCurve* Curve = FindOrAddCurve(CurveTable, AnimCurve.Name.DisplayName);
-		*Curve = AnimCurve.FloatCurve; // copy data
+		FRichCurve* Curve = FindOrAddCurve(CurveTable, FloatCurve.Name.DisplayName);
+		*Curve = FloatCurve.FloatCurve; // copy data
 	}
 
 	// we will need to add a curve to tell us the time we want to start playing audio
diff --git a/Engine/Source/Editor/AnimationModifiers/Private/AnimationBlueprintLibrary.cpp b/Engine/Source/Editor/AnimationModifiers/Private/AnimationBlueprintLibrary.cpp
--- a/Engine/Source/Editor/AnimationModifiers/Private/AnimationBlueprintLibrary.cpp
+++ b/Engine/Source/Editor/AnimationModifiers/Private/AnimationBlueprintLibrary.cpp
@@ -15,6 +15,11 @@
 #include "Algo/Transform.h"
 
 #include "AnimationRuntime.h"
+#include "Animation/AnimData/AnimDataModel.h"
+#include "Animation/AnimSequenceHelpers.h"
+#include "Animation/AnimData/AnimDataController.h"
+
+#define LOCTEXT_NAMESPACE "AnimationBlueprintLibrary"
 
 DEFINE_LOG_CATEGORY_STATIC(LogAnimationBlueprintLibrary, Verbose, All);
 
@@ -51,11 +56,11 @@ void UAnimationBlueprintLibrary::GetAnimationTrackNames(const UAnimSequence* Ani
 	TrackNames.Empty();
 	if (AnimationSequence)
 	{
-		TrackNames.Append(AnimationSequence->AnimationTrackNames);
+		AnimationSequence->GetDataModel()->GetBoneTrackNames(TrackNames);
 	}
 	else
 	{
-		UE_LOG(LogAnimationBlueprintLibrary, Warning, TEXT("Invalid Animation Sequence supplied for GetAnimationTrackNames"));
+		UE_LOG(LogAnimationBlueprintLibrary, Warning, TEXT("Invalid Animation Sequence supplied for GetBoneTrackNames"));
 	}	
 }
 
@@ -73,32 +78,26 @@ void UAnimationBlueprintLibrary::GetAnimationCurveNames(const UAnimSequence* Ani
 		{
 			case ERawCurveTrackTypes::RCT_Float:
 			{
-				Algo::Transform(AnimationSequence->RawCurveData.FloatCurves, CurveNames, GetCurveName);
-				break;
-			}
-
-			case ERawCurveTrackTypes::RCT_Vector:
-			{
-				Algo::Transform(AnimationSequence->RawCurveData.VectorCurves, CurveNames, GetCurveName);				
+				Algo::Transform(AnimationSequence->GetDataModel()->GetFloatCurves(), CurveNames, GetCurveName);
 				break;
 			}
 
 			case ERawCurveTrackTypes::RCT_Transform:
 			{
-				Algo::Transform(AnimationSequence->RawCurveData.TransformCurves, CurveNames, GetCurveName);
+				Algo::Transform(AnimationSequence->GetDataModel()->GetTransformCurves(), CurveNames, GetCurveName);
 				break;
 			}
 
 			default:
 			{
-				UE_LOG(LogAnimationBlueprintLibrary, Warning, TEXT("Invalid CurveType supplied for GetAnimationCurveNames"));
+				UE_LOG(LogAnimationBlueprintLibrary, Warning, TEXT("Invalid CurveType supplied for GetCurveNames"));
 			}
 		}
 		
 	}
 	else
 	{
-		UE_LOG(LogAnimationBlueprintLibrary, Warning, TEXT("Invalid Animation Sequence supplied for GetAnimationCurveNames"));
+		UE_LOG(LogAnimationBlueprintLibrary, Warning, TEXT("Invalid Animation Sequence supplied for GetCurveNames"));
 	}
 }
 
@@ -172,8 +171,8 @@ bool UAnimationBlueprintLibrary::IsValidRawAnimationTrackName(const UAnimSequenc
 
 	if (AnimationSequence)
 	{
-		const int32 TrackIndex = AnimationSequence->AnimationTrackNames.IndexOfByKey(TrackName);
-		bValidName = (TrackIndex != INDEX_NONE) && AnimationSequence->RawAnimationData.IsValidIndex(TrackIndex);
+		const int32 TrackIndex = AnimationSequence->GetDataModel()->GetBoneTrackIndexByName(TrackName);
+		bValidName = (TrackIndex != INDEX_NONE);
 	}
 	else
 	{
@@ -187,9 +186,8 @@ const FRawAnimSequenceTrack&  UAnimationBlueprintLibrary::GetRawAnimationTrackBy
 {
 	checkf(AnimationSequence, TEXT("Invalid Animation Sequence supplied for GetRawAnimationTrackByName"));
 
-	const int32 TrackIndex = AnimationSequence->AnimationTrackNames.IndexOfByKey(TrackName);
-	checkf(TrackIndex != INDEX_NONE, TEXT("Raw Animation Track %s does not exist in Animation Sequence %s"), *TrackName.ToString(), *AnimationSequence->GetName()); 
-	return AnimationSequence->GetRawAnimationTrack(TrackIndex);	
+	const FBoneAnimationTrack& AnimationTrack = AnimationSequence->GetDataModel()->GetBoneTrackByName(TrackName);
+	return AnimationTrack.InternalTrackData;
 }
 
 void UAnimationBlueprintLibrary::GetBoneCompressionSettings(const UAnimSequence* AnimationSequence, UAnimBoneCompressionSettings*& CompressionSettings)
@@ -927,7 +925,7 @@ void UAnimationBlueprintLibrary::CopyAnimNotifiesFromSequence(UAnimSequence* Src
 	else
 	{
 		const bool bShowDialogs = false;
-		UAnimSequence::CopyNotifies(SrcAnimSequence, DestAnimSequence, bShowDialogs);
+		UE::Anim::CopyNotifies(SrcAnimSequence, DestAnimSequence, bShowDialogs);
 	}
 }
 
@@ -1250,9 +1248,10 @@ void UAnimationBlueprintLibrary::RemoveAllCurveData(UAnimSequence* AnimationSequ
 {
 	if (AnimationSequence)
 	{
-		AnimationSequence->RawCurveData.DeleteAllCurveData(ERawCurveTrackTypes::RCT_Float);
-		AnimationSequence->RawCurveData.DeleteAllCurveData(ERawCurveTrackTypes::RCT_Vector);
-		AnimationSequence->RawCurveData.DeleteAllCurveData(ERawCurveTrackTypes::RCT_Transform);
+		UAnimDataController* Controller = AnimationSequence->GetController();
+
+		Controller->RemoveAllCurvesOfType(ERawCurveTrackTypes::RCT_Float);
+		Controller->RemoveAllCurvesOfType(ERawCurveTrackTypes::RCT_Transform);
 	}
 	else
 	{
@@ -1270,7 +1269,7 @@ void UAnimationBlueprintLibrary::AddTransformationCurveKey(UAnimSequence* Animat
 		TimeArray.Add(Time);
 		TransformArray.Add(Transform);
 
-		AddCurveKeysInternal<FTransform, FTransformCurve>(AnimationSequence, CurveName, TimeArray, TransformArray, ERawCurveTrackTypes::RCT_Transform);
+		AddCurveKeysInternal<FTransform, FTransformCurve, ERawCurveTrackTypes::RCT_Transform>(AnimationSequence, CurveName, TimeArray, TransformArray);
 	}
 	else
 	{
@@ -1285,7 +1284,7 @@ void UAnimationBlueprintLibrary::AddTransformationCurveKeys(UAnimSequence* Anima
 	{
 		if (Times.Num() == Transforms.Num())
 		{
-			AddCurveKeysInternal<FTransform, FTransformCurve>(AnimationSequence, CurveName, Times, Transforms, ERawCurveTrackTypes::RCT_Transform);
+			AddCurveKeysInternal<FTransform, FTransformCurve, ERawCurveTrackTypes::RCT_Transform>(AnimationSequence, CurveName, Times, Transforms);
 		}
 		else
 		{
@@ -1309,7 +1308,7 @@ void UAnimationBlueprintLibrary::AddFloatCurveKey(UAnimSequence* AnimationSequen
 		TimeArray.Add(Time);
 		ValueArray.Add(Value);
 
-		AddCurveKeysInternal<float, FFloatCurve>(AnimationSequence, CurveName, TimeArray, ValueArray, ERawCurveTrackTypes::RCT_Float);
+		AddCurveKeysInternal<float, FFloatCurve, ERawCurveTrackTypes::RCT_Float>(AnimationSequence, CurveName, TimeArray, ValueArray);
 	}
 	else
 	{
@@ -1324,7 +1323,7 @@ void UAnimationBlueprintLibrary::AddFloatCurveKeys(UAnimSequence* AnimationSeque
 	{
 		if (Times.Num() == Values.Num())
 		{
-			AddCurveKeysInternal<float, FFloatCurve>(AnimationSequence, CurveName, Times, Values, ERawCurveTrackTypes::RCT_Float);
+			AddCurveKeysInternal<float, FFloatCurve, ERawCurveTrackTypes::RCT_Float>(AnimationSequence, CurveName, Times, Values);
 		}
 		else
 		{
@@ -1349,7 +1348,7 @@ void UAnimationBlueprintLibrary::AddVectorCurveKey(UAnimSequence* AnimationSeque
 		TimeArray.Add(Time);
 		VectorArray.Add(Vector);
 
-		AddCurveKeysInternal<FVector, FVectorCurve>(AnimationSequence, CurveName, TimeArray, VectorArray, ERawCurveTrackTypes::RCT_Vector);
+		AddCurveKeysInternal<FVector, FVectorCurve, ERawCurveTrackTypes::RCT_Vector>(AnimationSequence, CurveName, TimeArray, VectorArray);
 	}
 	else
 	{
@@ -1364,7 +1363,7 @@ void UAnimationBlueprintLibrary::AddVectorCurveKeys(UAnimSequence* AnimationSequ
 	{
 		if (Times.Num() == Vectors.Num())
 		{
-			AddCurveKeysInternal<FVector, FVectorCurve>(AnimationSequence, CurveName, Times, Vectors, ERawCurveTrackTypes::RCT_Vector);
+			AddCurveKeysInternal<FVector, FVectorCurve, ERawCurveTrackTypes::RCT_Vector>(AnimationSequence, CurveName, Times, Vectors);
 		}
 		else
 		{
@@ -1375,12 +1374,35 @@ void UAnimationBlueprintLibrary::AddVectorCurveKeys(UAnimSequence* AnimationSequ
 	{
 		UE_LOG(LogAnimationBlueprintLibrary, Warning, TEXT("Invalid Animation Sequence for AddVectorCurveKeys"));
 	}
+}
+
+static void SetControllerCurveKeys(UAnimDataController* Controller, FSmartName SmartName, const TArray<float>& Times, const TArray<float>& Values)
+{
+	const int32 NumKeys = Values.Num();
+	const FAnimationCurveIdentifier CurveId(SmartName, ERawCurveTrackTypes::RCT_Float);
+	for (int32 KeyIndex = 0; KeyIndex < NumKeys; ++KeyIndex)
+	{
+		Controller->SetCurveKey(CurveId, { Times[KeyIndex], Values[KeyIndex] });
+	}
+}
 
+static void SetControllerCurveKeys(UAnimDataController* Controller, FSmartName SmartName, const TArray<float>& Times, const TArray<FTransform>& Values)
+{
+	const int32 NumKeys = Values.Num();
+	const FAnimationCurveIdentifier CurveId(SmartName, ERawCurveTrackTypes::RCT_Transform);
+	for (int32 KeyIndex = 0; KeyIndex < NumKeys; ++KeyIndex)
+	{
+		Controller->SetTransformCurveKey(CurveId, Times[KeyIndex], Values[KeyIndex]);
+	}
+}
 	
+static void SetControllerCurveKeys(UAnimDataController* Controller, FSmartName CurveId, const TArray<float>& Times, const TArray<FVector>& Values)
+{
+	ensure(false);
 }
 
-template <typename DataType, typename CurveClass>
-void UAnimationBlueprintLibrary::AddCurveKeysInternal(UAnimSequence* AnimationSequence, FName CurveName, const TArray<float>& Times, const TArray<DataType>& KeyData, ERawCurveTrackTypes CurveType)
+template <typename DataType, typename CurveClass, ERawCurveTrackTypes CurveType>
+void UAnimationBlueprintLibrary::AddCurveKeysInternal(UAnimSequence* AnimationSequence, FName CurveName, const TArray<float>& Times, const TArray<DataType>& KeyData)
 {
 	checkf(Times.Num() == KeyData.Num(), TEXT("Not enough key data supplied"));
 
@@ -1392,14 +1414,12 @@ void UAnimationBlueprintLibrary::AddCurveKeysInternal(UAnimSequence* AnimationSe
 		const FSmartName CurveSmartName = RetrieveSmartNameForCurve(AnimationSequence, CurveName, ContainerName);
 
 		// Retrieve the curve by name
-		CurveClass* Curve = static_cast<CurveClass*>(AnimationSequence->RawCurveData.GetCurveData(CurveSmartName.UID, CurveType));
+		const FAnimationCurveIdentifier CurveId(CurveSmartName, CurveType);
+		const CurveClass* Curve = static_cast<const CurveClass*>(AnimationSequence->GetDataModel()->FindCurve(CurveId));
 		if (Curve)
 		{
-			const int32 NumKeys = KeyData.Num();
-			for (int32 KeyIndex = 0; KeyIndex < NumKeys; ++KeyIndex)
-			{
-				Curve->UpdateOrAddKey(KeyData[KeyIndex], Times[KeyIndex]);
-			}
+			UAnimDataController* Controller = AnimationSequence->GetController();
+			SetControllerCurveKeys(Controller, CurveSmartName, Times, KeyData);
 		}
 	}
 }
@@ -1410,11 +1430,16 @@ bool UAnimationBlueprintLibrary::AddCurveInternal(UAnimSequence* AnimationSequen
 	FSmartName SmartCurveName;
 	AnimationSequence->GetSkeleton()->AddSmartNameAndModify(ContainerName, CurveName, SmartCurveName);	
 
+	FAnimationCurveIdentifier CurveId(SmartCurveName, SupportedCurveType);
+
 	bool bCurveAdded = false;
 
-	if (AnimationSequence->RawCurveData.GetCurveData(SmartCurveName.UID) == nullptr)
+	const FAnimCurveBase* ExistingCurve = AnimationSequence->GetDataModel()->FindCurve(CurveId);
+	if (ExistingCurve == nullptr)
 	{
-		bCurveAdded = AnimationSequence->RawCurveData.AddCurveData(SmartCurveName, CurveFlags, SupportedCurveType);		
+		UAnimDataController* Controller = AnimationSequence->GetController();
+		Controller->AddCurve(CurveId, CurveFlags);
+		bCurveAdded = AnimationSequence->GetDataModel()->FindCurve(CurveId) != nullptr;	
 	}
 	else
 	{
@@ -1436,15 +1461,10 @@ bool UAnimationBlueprintLibrary::RemoveCurveInternal(UAnimSequence* AnimationSeq
 		checkf(Skeleton != nullptr, TEXT("Invalid Skeleton ptr"));
 		if (Skeleton->GetSmartNameByUID(ContainerName, UID, SmartCurveName))
 		{
-			if (ContainerName == USkeleton::AnimTrackCurveMappingName)
-			{
-				bRemoved = AnimationSequence->RawCurveData.DeleteCurveData(SmartCurveName, ERawCurveTrackTypes::RCT_Transform);
-			}
-			else
-			{
-				bRemoved = AnimationSequence->RawCurveData.DeleteCurveData(SmartCurveName, ERawCurveTrackTypes::RCT_Float);
-				bRemoved |= AnimationSequence->RawCurveData.DeleteCurveData(SmartCurveName, ERawCurveTrackTypes::RCT_Vector);
-			}			
+			UAnimDataController* Controller = AnimationSequence->GetController();
+
+			const FAnimationCurveIdentifier CurveId(SmartCurveName, ContainerName == USkeleton::AnimTrackCurveMappingName ? ERawCurveTrackTypes::RCT_Transform : ERawCurveTrackTypes::RCT_Float);
+			bRemoved = Controller->RemoveCurve(CurveId);
 
 			if (bRemoveNameFromSkeleton)
 			{
@@ -1512,7 +1532,7 @@ void UAnimationBlueprintLibrary::GetFloatKeys(UAnimSequence* AnimationSequence,
 {
 	if (AnimationSequence)
 	{
-		GetCurveKeysInternal<float, FFloatCurve>(AnimationSequence, CurveName, Times, Values, ERawCurveTrackTypes::RCT_Float);
+		GetCurveKeysInternal<float, FFloatCurve, ERawCurveTrackTypes::RCT_Float>(AnimationSequence, CurveName, Times, Values);
 	}
 	else
 	{
@@ -1524,7 +1544,7 @@ void UAnimationBlueprintLibrary::GetVectorKeys(UAnimSequence* AnimationSequence,
 {
 	if (AnimationSequence)
 	{
-		GetCurveKeysInternal<FVector, FVectorCurve>(AnimationSequence, CurveName, Times, Values, ERawCurveTrackTypes::RCT_Vector);
+		GetCurveKeysInternal<FVector, FVectorCurve, ERawCurveTrackTypes::RCT_Vector>(AnimationSequence, CurveName, Times, Values);
 	}
 	else
 	{
@@ -1536,7 +1556,7 @@ void UAnimationBlueprintLibrary::GetTransformationKeys(UAnimSequence* AnimationS
 {
 	if (AnimationSequence)
 	{
-		GetCurveKeysInternal<FTransform, FTransformCurve>(AnimationSequence, CurveName, Times, Values, ERawCurveTrackTypes::RCT_Transform);
+		GetCurveKeysInternal<FTransform, FTransformCurve, ERawCurveTrackTypes::RCT_Transform>(AnimationSequence, CurveName, Times, Values);
 	}
 	else
 	{
@@ -1544,8 +1564,15 @@ void UAnimationBlueprintLibrary::GetTransformationKeys(UAnimSequence* AnimationS
 	}
 }
 
-template <typename DataType, typename CurveClass>
-void UAnimationBlueprintLibrary::GetCurveKeysInternal(UAnimSequence* AnimationSequence, FName CurveName, TArray<float>& Times, TArray<DataType>& KeyData, ERawCurveTrackTypes CurveType)
+void UAnimationBlueprintLibrary::CopyAnimationCurveNamesToSkeleton(USkeleton* OldSkeleton, USkeleton* NewSkeleton, UAnimSequenceBase* SequenceBase, ERawCurveTrackTypes CurveType)
+{
+	UAnimDataController* Controller = SequenceBase->GetController();
+	Controller->UpdateCurveNamesFromSkeleton(OldSkeleton, CurveType);
+	Controller->FindOrAddCurveNamesOnSkeleton(NewSkeleton, CurveType);	
+}
+
+template <typename DataType, typename CurveClass, ERawCurveTrackTypes CurveType>
+void UAnimationBlueprintLibrary::GetCurveKeysInternal(UAnimSequence* AnimationSequence, FName CurveName, TArray<float>& Times, TArray<DataType>& KeyData)
 {
 	checkf(AnimationSequence != nullptr, TEXT("Invalid Animation Sequence ptr"));
 	const FName ContainerName = RetrieveContainerNameForCurve(AnimationSequence, CurveName);
@@ -1556,7 +1583,9 @@ void UAnimationBlueprintLibrary::GetCurveKeysInternal(UAnimSequence* AnimationSe
 		const FSmartName CurveSmartName = RetrieveSmartNameForCurve(AnimationSequence, CurveName, ContainerName);
 
 		// Retrieve the curve by name
-		CurveClass* Curve = static_cast<CurveClass*>(AnimationSequence->RawCurveData.GetCurveData(CurveSmartName.UID, CurveType));
+		const FAnimationCurveIdentifier CurveId(CurveSmartName, CurveType);
+		const CurveClass* Curve = static_cast<const CurveClass*>(AnimationSequence->GetDataModel()->FindCurve(CurveId));
+
 		if (Curve)
 		{
 			Curve->GetKeys(Times, KeyData);
@@ -1574,13 +1603,15 @@ bool UAnimationBlueprintLibrary::DoesCurveExist(UAnimSequence* AnimationSequence
 		FSmartName SmartName;
 		if (RetrieveSmartNameForCurve(AnimationSequence, CurveName, USkeleton::AnimTrackCurveMappingName, SmartName))
 		{
-			FAnimCurveBase* Curve = AnimationSequence->RawCurveData.GetCurveData(SmartName.UID, CurveType);
+			FAnimationCurveIdentifier CurveId(SmartName, CurveType);
+			const FAnimCurveBase* Curve = AnimationSequence->GetDataModel()->FindCurve(CurveId);
 			bExistingCurve = Curve != nullptr;
 		}
 
 		if (RetrieveSmartNameForCurve(AnimationSequence, CurveName, USkeleton::AnimCurveMappingName, SmartName))
 		{
-			FAnimCurveBase* Curve = AnimationSequence->RawCurveData.GetCurveData(SmartName.UID, CurveType);
+			FAnimationCurveIdentifier CurveId(SmartName, CurveType);
+			const FAnimCurveBase* Curve = AnimationSequence->GetDataModel()->FindCurve(CurveId);
 			bExistingCurve |= Curve != nullptr;
 		}
 	}
@@ -1822,7 +1853,7 @@ void UAnimationBlueprintLibrary::GetBonePosesForTime(const UAnimSequence* Animat
 					FTransform& Transform = Poses[BoneNameIndex];
 					if (IsValidRawAnimationTrackName(AnimationSequence, BoneName))
 					{
-						AnimationSequence->ExtractBoneTransform(GetRawAnimationTrackByName(AnimationSequence, BoneName), Transform, Time);
+						UE::Anim::GetBoneTransformFromModel(AnimationSequence->GetDataModel(), Transform, AnimationSequence->GetDataModel()->GetBoneTrackIndexByName(BoneName), Time, AnimationSequence->Interpolation);
 					}
 					else
 					{
@@ -2130,12 +2161,13 @@ void UAnimationBlueprintLibrary::RemoveBoneAnimation(UAnimSequence* AnimationSeq
 {
 	if (AnimationSequence)
 	{
-		const TArray<FName>& TrackNames = AnimationSequence->GetAnimationTrackNames();
-		const int32 TrackIndex = TrackNames.Find(BoneName);
-		if (TrackIndex != INDEX_NONE)
+		TArray<FName> TrackNames;
+		AnimationSequence->GetDataModel()->GetBoneTrackNames(TrackNames);
+		
+		if (TrackNames.Contains(BoneName))
 		{
-			TArray<int32> TracksToRemove;
-			TracksToRemove.Add(TrackIndex);
+			TArray<FName> TracksToRemove;
+			TracksToRemove.Add(BoneName);
 
 			// remove all children if required
 			if (bIncludeChildren)
@@ -2147,31 +2179,26 @@ void UAnimationBlueprintLibrary::RemoveBoneAnimation(UAnimSequence* AnimationSeq
 					const int32 ParentBoneIndex = RefSkeleton.FindBoneIndex(BoneName);
 
 					// slow
-					for (int32 ChildTrackIndex = 0; ChildTrackIndex < TrackNames.Num(); ++ChildTrackIndex)
+					for (const FName& TrackName : TrackNames)
 					{
-						if (TrackIndex != ChildTrackIndex)
+						if (TrackName != BoneName)
 						{
-							const int32 ChildBoneIndex = RefSkeleton.FindBoneIndex(TrackNames[ChildTrackIndex]);
+							const int32 ChildBoneIndex = RefSkeleton.FindBoneIndex(TrackName);
 							if (RefSkeleton.BoneIsChildOf(ChildBoneIndex, ParentBoneIndex))
 							{
-								TracksToRemove.Add(ChildTrackIndex);
+								TracksToRemove.Add(TrackName);
 							}
 						}
 					}
 				}
 			}
 
-			TracksToRemove.Sort([](const int32& A, const int32& B) { return A < B; });
-
-			// go reverse since we're removing by index
-			for (int32 Index = TracksToRemove.Num() - 1; Index >= 0; --Index)
-			{
-				AnimationSequence->RemoveTrack(TracksToRemove[Index]);
-			}
+			UAnimDataController* Controller = AnimationSequence->GetController();
 
-			if (bFinalize)
+			UAnimDataController::FScopedBracket ScopedBracket(Controller, LOCTEXT("RemoveBoneAnimation_Bracket", "Removing Bone Animation Track"));
+			for (const FName& TrackName : TracksToRemove)
 			{
-				AnimationSequence->PostProcessSequence();
+				Controller->RemoveBoneTrack(TrackName);
 			}
 		}
 		else
@@ -2186,7 +2213,10 @@ void UAnimationBlueprintLibrary::RemoveAllBoneAnimation(UAnimSequence* Animation
 {
 	if (AnimationSequence)
 	{
-		AnimationSequence->RemoveAllTracks();
+		UAnimDataController* Controller = AnimationSequence->GetController();
+		UAnimDataController::FScopedBracket ScopedBracket(Controller, LOCTEXT("RemoveAllBoneAnimation_Bracket", "Removing all Bone Animation and Transform Curve Tracks"));
+		Controller->RemoveAllBoneTracks();
+		Controller->RemoveAllCurvesOfType(ERawCurveTrackTypes::RCT_Transform);		
 	}
 }
 
@@ -2194,15 +2224,18 @@ void UAnimationBlueprintLibrary::FinalizeBoneAnimation(UAnimSequence* AnimationS
 {
 	if (AnimationSequence)
 	{
+		PRAGMA_DISABLE_DEPRECATION_WARNINGS
 		AnimationSequence->PostProcessSequence();
+		PRAGMA_ENABLE_DEPRECATION_WARNINGS
 	}
 }
 
-template void UAnimationBlueprintLibrary::AddCurveKeysInternal<float, FFloatCurve>(UAnimSequence* AnimationSequence, FName CurveName, const TArray<float>& Times, const TArray<float>& KeyData, ERawCurveTrackTypes CurveType);
-template void UAnimationBlueprintLibrary::AddCurveKeysInternal<FVector, FVectorCurve>(UAnimSequence* AnimationSequence, FName CurveName, const TArray<float>& Times, const TArray<FVector>& KeyData, ERawCurveTrackTypes CurveType);
-template void UAnimationBlueprintLibrary::AddCurveKeysInternal<FTransform, FTransformCurve>(UAnimSequence* AnimationSequence, FName CurveName, const TArray<float>& Times, const TArray<FTransform>& KeyData, ERawCurveTrackTypes CurveType);
+template void UAnimationBlueprintLibrary::AddCurveKeysInternal<float, FFloatCurve, ERawCurveTrackTypes::RCT_Float>(UAnimSequence* AnimationSequence, FName CurveName, const TArray<float>& Times, const TArray<float>& KeyData);
+template void UAnimationBlueprintLibrary::AddCurveKeysInternal<FVector, FVectorCurve, ERawCurveTrackTypes::RCT_Vector>(UAnimSequence* AnimationSequence, FName CurveName, const TArray<float>& Times, const TArray<FVector>& KeyData);
+template void UAnimationBlueprintLibrary::AddCurveKeysInternal<FTransform, FTransformCurve, ERawCurveTrackTypes::RCT_Transform>(UAnimSequence* AnimationSequence, FName CurveName, const TArray<float>& Times, const TArray<FTransform>& KeyData);
 
-template void UAnimationBlueprintLibrary::GetCurveKeysInternal<float, FFloatCurve>(UAnimSequence* AnimationSequence, FName CurveName, TArray<float>& Times, TArray<float>& KeyData, ERawCurveTrackTypes CurveType);
-template void UAnimationBlueprintLibrary::GetCurveKeysInternal<FVector, FVectorCurve>(UAnimSequence* AnimationSequence, FName CurveName, TArray<float>& Times, TArray<FVector>& KeyData, ERawCurveTrackTypes CurveType);
-template void UAnimationBlueprintLibrary::GetCurveKeysInternal<FTransform, FTransformCurve>(UAnimSequence* AnimationSequence, FName CurveName, TArray<float>& Times, TArray<FTransform>& KeyData, ERawCurveTrackTypes CurveType);
+template void UAnimationBlueprintLibrary::GetCurveKeysInternal<float, FFloatCurve, ERawCurveTrackTypes::RCT_Float>(UAnimSequence* AnimationSequence, FName CurveName, TArray<float>& Times, TArray<float>& KeyData);
+template void UAnimationBlueprintLibrary::GetCurveKeysInternal<FVector, FVectorCurve, ERawCurveTrackTypes::RCT_Vector>(UAnimSequence* AnimationSequence, FName CurveName, TArray<float>& Times, TArray<FVector>& KeyData);
+template void UAnimationBlueprintLibrary::GetCurveKeysInternal<FTransform, FTransformCurve, ERawCurveTrackTypes::RCT_Transform>(UAnimSequence* AnimationSequence, FName CurveName, TArray<float>& Times, TArray<FTransform>& KeyData);
 
+#undef LOCTEXT_NAMESPACE // "AnimationBlueprintLibrary"
\ No newline at end of file
diff --git a/Engine/Source/Editor/AnimationModifiers/Private/AnimationModifier.cpp b/Engine/Source/Editor/AnimationModifiers/Private/AnimationModifier.cpp
--- a/Engine/Source/Editor/AnimationModifiers/Private/AnimationModifier.cpp
+++ b/Engine/Source/Editor/AnimationModifiers/Private/AnimationModifier.cpp
@@ -13,6 +13,8 @@
 #include "UObject/UObjectIterator.h"
 #include "UObject/AnimObjectVersion.h"
 
+#define LOCTEXT_NAMESPACE "AnimationModifier"
+
 UAnimationModifier::UAnimationModifier()
 	: PreviouslyAppliedModifier(nullptr)
 {
@@ -44,13 +46,17 @@ void UAnimationModifier::ApplyToAnimationSequence(class UAnimSequence* InAnimati
 	/** In case this modifier has been previously applied, revert it using the serialised out version at the time */	
 	if (PreviouslyAppliedModifier)
 	{
-
 		PreviouslyAppliedModifier->Modify();
 		PreviouslyAppliedModifier->OnRevert(CurrentAnimSequence);
 	}
 
-	/** Reverting and applying, populates the log with possible warnings and or errors to notify the user about */
-	OnApply(CurrentAnimSequence);
+	UAnimDataController* Controller = CurrentAnimSequence->GetController();
+
+	{
+		UAnimDataController::FScopedBracket ScopedBracket(Controller, LOCTEXT("ApplyModifierBracket", "Applying Animation Modifier"));
+		/** Reverting and applying, populates the log with possible warnings and or errors to notify the user about */
+		OnApply(CurrentAnimSequence);
+	}
 
 	// Apply transaction
 	ModifierTransaction.BeginOperation();
@@ -80,12 +86,9 @@ void UAnimationModifier::ApplyToAnimationSequence(class UAnimSequence* InAnimati
 		AnimationDataTransaction.Apply();
 		AnimationDataTransaction.EndOperation();
 		CurrentAnimSequence->RefreshCacheData();
-		CurrentAnimSequence->RefreshCurveData();
 	}
 	else
-	{		
-		UpdateCompressedAnimationData();
-		
+	{
 		/** Mark the previous modifier pending kill, as it will be replaced with the current modifier state */
 		if (PreviouslyAppliedModifier)
 		{
@@ -97,8 +100,6 @@ void UAnimationModifier::ApplyToAnimationSequence(class UAnimSequence* InAnimati
 		CurrentAnimSequence->PostEditChange();
 		CurrentSkeleton->PostEditChange();
 		CurrentAnimSequence->RefreshCacheData();
-		CurrentAnimSequence->RefreshCurveData();
-		CurrentAnimSequence->MarkRawDataAsModified();
 
 		UpdateStoredRevisions();
 	}
@@ -132,20 +133,22 @@ void UAnimationModifier::RevertFromAnimationSequence(class UAnimSequence* InAnim
 		Transaction.SaveObject(this);
 
 		PreviouslyAppliedModifier->Modify();
-		PreviouslyAppliedModifier->OnRevert(CurrentAnimSequence);
+
+		UAnimDataController* Controller = CurrentAnimSequence->GetController();
+
+		{
+			UAnimDataController::FScopedBracket ScopedBracket(Controller, LOCTEXT("RevertModifierBracket", "Reverting Animation Modifier"));
+			PreviouslyAppliedModifier->OnRevert(CurrentAnimSequence);
+		}
 
 		// Apply transaction
 		Transaction.BeginOperation();
 		Transaction.Apply();
 		Transaction.EndOperation();
 
-		UpdateCompressedAnimationData();
-
 	    CurrentAnimSequence->PostEditChange();
 	    CurrentSkeleton->PostEditChange();
 	    CurrentAnimSequence->RefreshCacheData();
-	    CurrentAnimSequence->RefreshCurveData();
-	    CurrentAnimSequence->MarkRawDataAsModified();
 
 		ResetStoredRevisions();
 
@@ -254,3 +257,5 @@ void UAnimationModifier::SetInstanceRevisionGuid(FGuid Guid)
 {
 	RevisionGuid = Guid;
 }
+
+#undef LOCTEXT_NAMESPACE
\ No newline at end of file
diff --git a/Engine/Source/Editor/AnimationModifiers/Public/AnimationBlueprintLibrary.h b/Engine/Source/Editor/AnimationModifiers/Public/AnimationBlueprintLibrary.h
--- a/Engine/Source/Editor/AnimationModifiers/Public/AnimationBlueprintLibrary.h
+++ b/Engine/Source/Editor/AnimationModifiers/Public/AnimationBlueprintLibrary.h
@@ -297,8 +297,8 @@ public:
 	static void AddVectorCurveKeys(UAnimSequence* AnimationSequence, FName CurveName, const TArray<float>& Times, const TArray<FVector>& Vectors);
 
 	// Curve helper functions
-	template <typename DataType, typename CurveClass>
-	static void AddCurveKeysInternal(UAnimSequence* AnimationSequence, FName CurveName, const TArray<float>& Times, const TArray<DataType>& KeyData, ERawCurveTrackTypes CurveType);
+	template <typename DataType, typename CurveClass, ERawCurveTrackTypes CurveType>
+	static void AddCurveKeysInternal(UAnimSequence* AnimationSequence, FName CurveName, const TArray<float>& Times, const TArray<DataType>& KeyData);
 
 	// Returns true if successfully added, false if it was already existing
 	static bool AddCurveInternal(UAnimSequence* AnimationSequence, FName CurveName, FName ContainerName, int32 CurveFlags, ERawCurveTrackTypes SupportedCurveType);
@@ -322,8 +322,12 @@ public:
 	UFUNCTION(BlueprintCallable, Category = "AnimationBlueprintLibrary|Curves")
 	static void GetTransformationKeys(UAnimSequence* AnimationSequence, FName CurveName, TArray<float>& Times, TArray<FTransform>& Values);
 
-	template <typename DataType, typename CurveClass>
-	static void GetCurveKeysInternal(UAnimSequence* AnimationSequence, FName CurveName, TArray<float>& Times, TArray<DataType>& KeyData, ERawCurveTrackTypes CurveType);
+	template <typename DataType, typename CurveClass, ERawCurveTrackTypes CurveType>
+	static void GetCurveKeysInternal(UAnimSequence* AnimationSequence, FName CurveName, TArray<float>& Times, TArray<DataType>& KeyData);
+
+	/** Ensures that any curve names that do not exist on the NewSkeleton are added to it, in which case the SmartName on the actual curve itself will also be updated */
+	UFUNCTION(BlueprintCallable, Category = "AnimationBlueprintLibrary|Curves")
+	static void CopyAnimationCurveNamesToSkeleton(USkeleton* OldSkeleton, USkeleton* NewSkeleton, UAnimSequenceBase* SequenceBase, ERawCurveTrackTypes CurveType);
 	
 	// Bone Tracks
 
@@ -342,7 +346,8 @@ public:
 	static void RemoveAllBoneAnimation(UAnimSequence* AnimationSequence);
 
 	/** Apply all the changes made to Bone Tracks to Finalize. This triggers recompression. Note that this is expensive, but will require to get correct compressed data */
-	UFUNCTION(BlueprintCallable, Category = "AnimationBlueprintLibrary|Curves")
+	UE_DEPRECATED(5.0, "FinalizeBoneAnimation has been deprecated, use UAnimDataController instead")
+	UFUNCTION(BlueprintCallable, Category = "AnimationBlueprintLibrary|Curves", meta=(DeprecatedFunction, DeprecationMessage="FinalizeBoneAnimation has been deprecated, use UAnimDataController instead"))
 	static void FinalizeBoneAnimation(UAnimSequence* AnimationSequence);
 
 	// Smart name helper functions
diff --git a/Engine/Source/Editor/DetailCustomizations/Private/AnimSequenceDetails.cpp b/Engine/Source/Editor/DetailCustomizations/Private/AnimSequenceDetails.cpp
--- a/Engine/Source/Editor/DetailCustomizations/Private/AnimSequenceDetails.cpp
+++ b/Engine/Source/Editor/DetailCustomizations/Private/AnimSequenceDetails.cpp
@@ -619,7 +619,7 @@ TArray<float> SAnimationRefPoseViewport::GetBars() const
 	TArray<float> Bars;
 	if (AnimRef)
 	{
-		int RefFrameIndex;
+		int32 RefFrameIndex;
 		RefFrameIndexPropertyHandle->GetValue(RefFrameIndex);
 		float Fraction = (AnimRef->GetNumberOfSampledKeys() > 0)? FMath::Clamp<float>((float)RefFrameIndex/(float)AnimRef->GetNumberOfSampledKeys(), 0.f, 1.f) : 0.f;
 		Bars.Add(AnimRef->GetPlayLength() * Fraction);
diff --git a/Engine/Source/Editor/MovieSceneTools/Private/TrackEditors/TransformTrackEditor.cpp b/Engine/Source/Editor/MovieSceneTools/Private/TrackEditors/TransformTrackEditor.cpp
--- a/Engine/Source/Editor/MovieSceneTools/Private/TrackEditors/TransformTrackEditor.cpp
+++ b/Engine/Source/Editor/MovieSceneTools/Private/TrackEditors/TransformTrackEditor.cpp
@@ -27,6 +27,7 @@
 #include "TransformPropertySection.h"
 #include "SequencerUtilities.h"
 #include "MovieSceneToolHelpers.h"
+#include "Animation/AnimData/AnimDataModel.h"
 
 #include "EntitySystem/Interrogation/MovieSceneInterrogationLinker.h"
 #include "EntitySystem/Interrogation/MovieSceneInterrogatedPropertyInstantiator.h"
@@ -1005,7 +1006,7 @@ void F3DTransformTrackEditor::ImportAnimSequenceTransforms(const FAssetData& Ass
 		}
 	}
 
-	if(AnimSequence && AnimSequence->GetRawAnimationData().Num() > 0)
+	if(AnimSequence && AnimSequence->GetDataModel()->GetNumBoneTracks() > 0)
 	{
 		const FScopedTransaction Transaction( NSLOCTEXT( "Sequencer", "ImportAnimSequenceTransforms", "Import Anim Sequence Transforms" ) );
 
@@ -1044,7 +1045,9 @@ void F3DTransformTrackEditor::ImportAnimSequenceTransforms(const FAssetData& Ass
 
 			TArray<FTempTransformKey> TempKeys;
 
-			FRawAnimSequenceTrack& RawTrack = AnimSequence->GetRawAnimationTrack(0);
+			const FBoneAnimationTrack& AnimationTrack = AnimSequence->GetDataModel()->GetBoneTrackByIndex(0);
+			const FRawAnimSequenceTrack& RawTrack = AnimationTrack.InternalTrackData;
+
 			const int32 KeyCount = FMath::Max(FMath::Max(RawTrack.PosKeys.Num(), RawTrack.RotKeys.Num()), RawTrack.ScaleKeys.Num());
 			for(int32 KeyIndex = 0; KeyIndex < KeyCount; KeyIndex++)
 			{
diff --git a/Engine/Source/Editor/Persona/Private/AnimTimeline/AnimModel_AnimComposite.cpp b/Engine/Source/Editor/Persona/Private/AnimTimeline/AnimModel_AnimComposite.cpp
--- a/Engine/Source/Editor/Persona/Private/AnimTimeline/AnimModel_AnimComposite.cpp
+++ b/Engine/Source/Editor/Persona/Private/AnimTimeline/AnimModel_AnimComposite.cpp
@@ -79,9 +79,9 @@ void FAnimModel_AnimComposite::RecalculateSequenceLength()
 		{
 			ClampToEndTime(NewSequenceLength);
 
-			AnimComposite->SetSequenceLength(NewSequenceLength);
+			AnimComposite->SetCompositeLength(NewSequenceLength);
 
-			// Reset view if we changed length (note: has to be done after ->SetSequenceLength)!
+			// Reset view if we changed length (note: has to be done after ->SetCompositeLength)!
 			UpdateRange();
 
 			UAnimPreviewInstance* PreviewInstance = (GetPreviewScene()->GetPreviewMeshComponent()) ? GetPreviewScene()->GetPreviewMeshComponent()->PreviewInstance : nullptr;
diff --git a/Engine/Source/Editor/Persona/Private/AnimTimeline/AnimModel_AnimMontage.cpp b/Engine/Source/Editor/Persona/Private/AnimTimeline/AnimModel_AnimMontage.cpp
--- a/Engine/Source/Editor/Persona/Private/AnimTimeline/AnimModel_AnimMontage.cpp
+++ b/Engine/Source/Editor/Persona/Private/AnimTimeline/AnimModel_AnimMontage.cpp
@@ -232,9 +232,9 @@ void FAnimModel_AnimMontage::RecalculateSequenceLength()
 
 			RefreshSectionTimes();
 
-			AnimMontage->SetSequenceLength(NewSequenceLength);
+			AnimMontage->SetCompositeLength(NewSequenceLength);
 
-			// Reset view if we changed length (note: has to be done after ->SetSequenceLength)!
+			// Reset view if we changed length (note: has to be done after ->SetCompositeLength)!
 			UpdateRange();
 
 			UAnimPreviewInstance* PreviewInstance = (GetPreviewScene()->GetPreviewMeshComponent()) ? GetPreviewScene()->GetPreviewMeshComponent()->PreviewInstance : nullptr;
diff --git a/Engine/Source/Editor/Persona/Private/AnimTimeline/AnimModel_AnimSequenceBase.cpp b/Engine/Source/Editor/Persona/Private/AnimTimeline/AnimModel_AnimSequenceBase.cpp
--- a/Engine/Source/Editor/Persona/Private/AnimTimeline/AnimModel_AnimSequenceBase.cpp
+++ b/Engine/Source/Editor/Persona/Private/AnimTimeline/AnimModel_AnimSequenceBase.cpp
@@ -21,6 +21,8 @@
 #include "AnimPreviewInstance.h"
 #include "ScopedTransaction.h"
 
+#include "Animation/AnimData/AnimDataController.h"
+
 #define LOCTEXT_NAMESPACE "FAnimModel_AnimSequence"
 
 FAnimModel_AnimSequenceBase::FAnimModel_AnimSequenceBase(const TSharedRef<IPersonaPreviewScene>& InPreviewScene, const TSharedRef<IEditableSkeleton>& InEditableSkeleton, const TSharedRef<FUICommandList>& InCommandList, UAnimSequenceBase* InAnimSequenceBase)
@@ -38,13 +40,14 @@ FAnimModel_AnimSequenceBase::FAnimModel_AnimSequenceBase(const TSharedRef<IPerso
 		bElementNodeDisplayFlag = false;
 	}
 
-	AnimSequenceBase->RegisterOnAnimTrackCurvesChanged(UAnimSequenceBase::FOnAnimTrackCurvesChanged::CreateRaw(this, &FAnimModel_AnimSequenceBase::RefreshTracks));
 	AnimSequenceBase->RegisterOnNotifyChanged(UAnimSequenceBase::FOnNotifyChanged::CreateRaw(this, &FAnimModel_AnimSequenceBase::RefreshSnapTimes));
 	
 	if(GEditor)
 	{
 		GEditor->RegisterForUndo(this);
 	}
+
+	AnimSequenceBase->GetDataModel()->GetModifiedEvent().AddRaw(this, &FAnimModel_AnimSequenceBase::OnDataModelChanged);
 }
 
 FAnimModel_AnimSequenceBase::~FAnimModel_AnimSequenceBase()
@@ -55,7 +58,8 @@ FAnimModel_AnimSequenceBase::~FAnimModel_AnimSequenceBase()
 	}
 
 	AnimSequenceBase->UnregisterOnNotifyChanged(this);
-	AnimSequenceBase->UnregisterOnAnimTrackCurvesChanged(this);
+
+	AnimSequenceBase->GetDataModel()->GetModifiedEvent().RemoveAll(this);
 }
 
 void FAnimModel_AnimSequenceBase::Initialize()
@@ -167,8 +171,6 @@ UAnimSequenceBase* FAnimModel_AnimSequenceBase::GetAnimSequenceBase() const
 
 void FAnimModel_AnimSequenceBase::RefreshNotifyTracks()
 {
-	AnimSequenceBase->InitializeNotifyTrack();
-
 	if(!NotifyRoot.IsValid())
 	{
 		// Add a root track for notifies & then the main 'panel' legacy track
@@ -199,9 +201,10 @@ void FAnimModel_AnimSequenceBase::RefreshCurveTracks()
 	RootTracks.Add(CurveRoot.ToSharedRef());
 
 	// Next add a track for each float curve
-	for(FFloatCurve& FloatCurve : AnimSequenceBase->RawCurveData.FloatCurves)
+	const FAnimationCurveData& AnimationModelCurveData = AnimSequenceBase->GetDataModel()->GetCurveData();
+	for(const FFloatCurve& FloatCurve : AnimationModelCurveData.FloatCurves)
 	{
-		CurveRoot->AddChild(MakeShared<FAnimTimelineTrack_FloatCurve>(FloatCurve, SharedThis(this)));
+		CurveRoot->AddChild(MakeShared<FAnimTimelineTrack_FloatCurve>(&FloatCurve, SharedThis(this)));
 	}
 
 	UAnimSequence* AnimSequence = Cast<UAnimSequence>(AnimSequenceBase);
@@ -217,9 +220,9 @@ void FAnimModel_AnimSequenceBase::RefreshCurveTracks()
 		RootTracks.Add(AdditiveRoot.ToSharedRef());
 
 		// Next add a track for each transform curve
-		for(FTransformCurve& TransformCurve : AnimSequence->RawCurveData.TransformCurves)
+		for(const FTransformCurve& TransformCurve : AnimationModelCurveData.TransformCurves)
 		{
-			TSharedRef<FAnimTimelineTrack_TransformCurve> TransformCurveTrack = MakeShared<FAnimTimelineTrack_TransformCurve>(TransformCurve, SharedThis(this));
+			TSharedRef<FAnimTimelineTrack_TransformCurve> TransformCurveTrack = MakeShared<FAnimTimelineTrack_TransformCurve>(&TransformCurve, SharedThis(this));
 			TransformCurveTrack->SetExpanded(false);
 			AdditiveRoot->AddChild(TransformCurveTrack);
 
@@ -234,33 +237,64 @@ void FAnimModel_AnimSequenceBase::RefreshCurveTracks()
 			
 			FText VectorFormat = LOCTEXT("TransformVectorFormat", "{0}.{1}");
 			FText TranslationName = LOCTEXT("TransformTranslationTrackName", "Translation");
-			TSharedRef<FAnimTimelineTrack_VectorCurve> TranslationCurveTrack = MakeShared<FAnimTimelineTrack_VectorCurve>(TransformCurve.TranslationCurve, TransformCurve.Name, 0, ERawCurveTrackTypes::RCT_Transform, TranslationName, FText::Format(VectorFormat, TransformName, TranslationName), TransformColor, SharedThis(this));
+			TSharedRef<FAnimTimelineTrack_VectorCurve> TranslationCurveTrack = MakeShared<FAnimTimelineTrack_VectorCurve>(&TransformCurve.TranslationCurve, TransformCurve.Name, 0, ERawCurveTrackTypes::RCT_Transform, TranslationName, FText::Format(VectorFormat, TransformName, TranslationName), TransformColor, SharedThis(this));
 			TranslationCurveTrack->SetExpanded(false);
 			TransformCurveTrack->AddChild(TranslationCurveTrack);			
 
 			FText ComponentFormat = LOCTEXT("TransformComponentFormat", "{0}.{1}.{2}");
-			TranslationCurveTrack->AddChild(MakeShared<FAnimTimelineTrack_Curve>(TransformCurve.TranslationCurve.FloatCurves[0], TransformCurve.Name, 0, ERawCurveTrackTypes::RCT_Transform, XName, FText::Format(ComponentFormat, TransformName, TranslationName, XName), XColor, XColor, SharedThis(this)));
-			TranslationCurveTrack->AddChild(MakeShared<FAnimTimelineTrack_Curve>(TransformCurve.TranslationCurve.FloatCurves[1], TransformCurve.Name, 1, ERawCurveTrackTypes::RCT_Transform, YName, FText::Format(ComponentFormat, TransformName, TranslationName, YName), YColor, YColor, SharedThis(this)));
-			TranslationCurveTrack->AddChild(MakeShared<FAnimTimelineTrack_Curve>(TransformCurve.TranslationCurve.FloatCurves[2], TransformCurve.Name, 2, ERawCurveTrackTypes::RCT_Transform, ZName, FText::Format(ComponentFormat, TransformName, TranslationName, ZName), ZColor, ZColor, SharedThis(this)));
+			TranslationCurveTrack->AddChild(MakeShared<FAnimTimelineTrack_Curve>(&TransformCurve.TranslationCurve.FloatCurves[0], TransformCurve.Name, 0, ERawCurveTrackTypes::RCT_Transform, XName, FText::Format(ComponentFormat, TransformName, TranslationName, XName), XColor, XColor, SharedThis(this)));
+			TranslationCurveTrack->AddChild(MakeShared<FAnimTimelineTrack_Curve>(&TransformCurve.TranslationCurve.FloatCurves[1], TransformCurve.Name, 1, ERawCurveTrackTypes::RCT_Transform, YName, FText::Format(ComponentFormat, TransformName, TranslationName, YName), YColor, YColor, SharedThis(this)));
+			TranslationCurveTrack->AddChild(MakeShared<FAnimTimelineTrack_Curve>(&TransformCurve.TranslationCurve.FloatCurves[2], TransformCurve.Name, 2, ERawCurveTrackTypes::RCT_Transform, ZName, FText::Format(ComponentFormat, TransformName, TranslationName, ZName), ZColor, ZColor, SharedThis(this)));
 
 			FText RollName = LOCTEXT("RotationRollTrackName", "Roll");
 			FText PitchName = LOCTEXT("RotationPitchTrackName", "Pitch");
 			FText YawName = LOCTEXT("RotationYawTrackName", "Yaw");
 			FText RotationName = LOCTEXT("TransformRotationTrackName", "Rotation");
-			TSharedRef<FAnimTimelineTrack_VectorCurve> RotationCurveTrack = MakeShared<FAnimTimelineTrack_VectorCurve>(TransformCurve.RotationCurve, TransformCurve.Name, 3, ERawCurveTrackTypes::RCT_Transform, RotationName, FText::Format(VectorFormat, TransformName, RotationName), TransformColor, SharedThis(this));
+			TSharedRef<FAnimTimelineTrack_VectorCurve> RotationCurveTrack = MakeShared<FAnimTimelineTrack_VectorCurve>(&TransformCurve.RotationCurve, TransformCurve.Name, 3, ERawCurveTrackTypes::RCT_Transform, RotationName, FText::Format(VectorFormat, TransformName, RotationName), TransformColor, SharedThis(this));
 			RotationCurveTrack->SetExpanded(false);
 			TransformCurveTrack->AddChild(RotationCurveTrack);
-			RotationCurveTrack->AddChild(MakeShared<FAnimTimelineTrack_Curve>(TransformCurve.RotationCurve.FloatCurves[0], TransformCurve.Name, 3, ERawCurveTrackTypes::RCT_Transform, RollName, FText::Format(ComponentFormat, TransformName, RotationName, RollName), XColor, XColor, SharedThis(this)));
-			RotationCurveTrack->AddChild(MakeShared<FAnimTimelineTrack_Curve>(TransformCurve.RotationCurve.FloatCurves[1], TransformCurve.Name, 4, ERawCurveTrackTypes::RCT_Transform, PitchName, FText::Format(ComponentFormat, TransformName, RotationName, PitchName), YColor, YColor, SharedThis(this)));
-			RotationCurveTrack->AddChild(MakeShared<FAnimTimelineTrack_Curve>(TransformCurve.RotationCurve.FloatCurves[2], TransformCurve.Name, 5, ERawCurveTrackTypes::RCT_Transform, YawName, FText::Format(ComponentFormat, TransformName, RotationName, YawName), ZColor, ZColor, SharedThis(this)));
+			RotationCurveTrack->AddChild(MakeShared<FAnimTimelineTrack_Curve>(&TransformCurve.RotationCurve.FloatCurves[0], TransformCurve.Name, 3, ERawCurveTrackTypes::RCT_Transform, RollName, FText::Format(ComponentFormat, TransformName, RotationName, RollName), XColor, XColor, SharedThis(this)));
+			RotationCurveTrack->AddChild(MakeShared<FAnimTimelineTrack_Curve>(&TransformCurve.RotationCurve.FloatCurves[1], TransformCurve.Name, 4, ERawCurveTrackTypes::RCT_Transform, PitchName, FText::Format(ComponentFormat, TransformName, RotationName, PitchName), YColor, YColor, SharedThis(this)));
+			RotationCurveTrack->AddChild(MakeShared<FAnimTimelineTrack_Curve>(&TransformCurve.RotationCurve.FloatCurves[2], TransformCurve.Name, 5, ERawCurveTrackTypes::RCT_Transform, YawName, FText::Format(ComponentFormat, TransformName, RotationName, YawName), ZColor, ZColor, SharedThis(this)));
 
 			FText ScaleName = LOCTEXT("TransformScaleTrackName", "Scale");
-			TSharedRef<FAnimTimelineTrack_VectorCurve> ScaleCurveTrack = MakeShared<FAnimTimelineTrack_VectorCurve>(TransformCurve.ScaleCurve, TransformCurve.Name, 6, ERawCurveTrackTypes::RCT_Transform, ScaleName, FText::Format(VectorFormat, TransformName, ScaleName), TransformColor, SharedThis(this));
+			TSharedRef<FAnimTimelineTrack_VectorCurve> ScaleCurveTrack = MakeShared<FAnimTimelineTrack_VectorCurve>(&TransformCurve.ScaleCurve, TransformCurve.Name, 6, ERawCurveTrackTypes::RCT_Transform, ScaleName, FText::Format(VectorFormat, TransformName, ScaleName), TransformColor, SharedThis(this));
 			ScaleCurveTrack->SetExpanded(false);
 			TransformCurveTrack->AddChild(ScaleCurveTrack);
-			ScaleCurveTrack->AddChild(MakeShared<FAnimTimelineTrack_Curve>(TransformCurve.ScaleCurve.FloatCurves[0], TransformCurve.Name, 6, ERawCurveTrackTypes::RCT_Transform, XName, FText::Format(ComponentFormat, TransformName, ScaleName, XName), XColor, XColor, SharedThis(this)));
-			ScaleCurveTrack->AddChild(MakeShared<FAnimTimelineTrack_Curve>(TransformCurve.ScaleCurve.FloatCurves[1], TransformCurve.Name, 7, ERawCurveTrackTypes::RCT_Transform, YName, FText::Format(ComponentFormat, TransformName, ScaleName, YName), YColor, YColor, SharedThis(this)));
-			ScaleCurveTrack->AddChild(MakeShared<FAnimTimelineTrack_Curve>(TransformCurve.ScaleCurve.FloatCurves[2], TransformCurve.Name, 8, ERawCurveTrackTypes::RCT_Transform, ZName, FText::Format(ComponentFormat, TransformName, ScaleName, ZName), ZColor, ZColor, SharedThis(this)));
+			ScaleCurveTrack->AddChild(MakeShared<FAnimTimelineTrack_Curve>(&TransformCurve.ScaleCurve.FloatCurves[0], TransformCurve.Name, 6, ERawCurveTrackTypes::RCT_Transform, XName, FText::Format(ComponentFormat, TransformName, ScaleName, XName), XColor, XColor, SharedThis(this)));
+			ScaleCurveTrack->AddChild(MakeShared<FAnimTimelineTrack_Curve>(&TransformCurve.ScaleCurve.FloatCurves[1], TransformCurve.Name, 7, ERawCurveTrackTypes::RCT_Transform, YName, FText::Format(ComponentFormat, TransformName, ScaleName, YName), YColor, YColor, SharedThis(this)));
+			ScaleCurveTrack->AddChild(MakeShared<FAnimTimelineTrack_Curve>(&TransformCurve.ScaleCurve.FloatCurves[2], TransformCurve.Name, 8, ERawCurveTrackTypes::RCT_Transform, ZName, FText::Format(ComponentFormat, TransformName, ScaleName, ZName), ZColor, ZColor, SharedThis(this)));
+		}
+	}
+}
+
+void FAnimModel_AnimSequenceBase::OnDataModelChanged(const EAnimDataModelNotifyType& NotifyType, UAnimDataModel* Model, const FAnimDataModelNotifPayload& PayLoad)
+{
+	NotifyCollector.Handle(NotifyType);
+
+	switch(NotifyType)
+	{ 
+		case EAnimDataModelNotifyType::CurveAdded:
+		case EAnimDataModelNotifyType::CurveRemoved:
+		case EAnimDataModelNotifyType::TrackAdded:
+		case EAnimDataModelNotifyType::TrackChanged:
+		case EAnimDataModelNotifyType::TrackRemoved:
+		case EAnimDataModelNotifyType::SequenceLengthChanged:
+		case EAnimDataModelNotifyType::FrameRateChanged:
+		{
+			if (NotifyCollector.IsNotWithinBracket())
+			{
+				RefreshTracks();
+			}
+			break;
+		}
+		case EAnimDataModelNotifyType::BracketClosed:
+		{
+			if (NotifyCollector.IsNotWithinBracket())
+			{
+				RefreshTracks();
+			}
+			break;
 		}
 	}
 }
@@ -273,7 +307,7 @@ void FAnimModel_AnimSequenceBase::EditSelectedCurves()
 		if(SelectedTrack->IsA<FAnimTimelineTrack_Curve>())
 		{
 			TSharedRef<FAnimTimelineTrack_Curve> CurveTrack = StaticCastSharedRef<FAnimTimelineTrack_Curve>(SelectedTrack);
-			const TArray<FRichCurve*> Curves = CurveTrack->GetCurves();
+			const TArray<const FRichCurve*> Curves = CurveTrack->GetCurves();
 			int32 NumCurves = Curves.Num();
 			for(int32 CurveIndex = 0; CurveIndex < NumCurves; ++CurveIndex)
 			{
@@ -305,7 +339,7 @@ bool FAnimModel_AnimSequenceBase::CanEditSelectedCurves() const
 		if(SelectedTrack->IsA<FAnimTimelineTrack_Curve>())
 		{
 			TSharedRef<FAnimTimelineTrack_Curve> CurveTrack = StaticCastSharedRef<FAnimTimelineTrack_Curve>(SelectedTrack);
-			const TArray<FRichCurve*>& Curves = CurveTrack->GetCurves();
+			const TArray<const FRichCurve*>& Curves = CurveTrack->GetCurves();
 			for(int32 CurveIndex = 0; CurveIndex < Curves.Num(); ++CurveIndex)
 			{
 				if(CurveTrack->CanEditCurve(CurveIndex))
@@ -321,9 +355,8 @@ bool FAnimModel_AnimSequenceBase::CanEditSelectedCurves() const
 
 void FAnimModel_AnimSequenceBase::RemoveSelectedCurves()
 {
-	FScopedTransaction Transaction(LOCTEXT("CurvePanel_RemoveCurves", "Remove Curves"));
-
-	AnimSequenceBase->Modify(true);
+	UAnimDataController* Controller = AnimSequenceBase->GetController();
+	Controller->OpenBracket(LOCTEXT("CurvePanel_RemoveCurves", "Remove Curves"));
 
 	bool bDeletedCurve = false;
 
@@ -333,10 +366,9 @@ void FAnimModel_AnimSequenceBase::RemoveSelectedCurves()
 		{
 			TSharedRef<FAnimTimelineTrack_FloatCurve> FloatCurveTrack = StaticCastSharedRef<FAnimTimelineTrack_FloatCurve>(SelectedTrack);
 
-			FFloatCurve& FloatCurve = FloatCurveTrack->GetFloatCurve();
 			FSmartName CurveName = FloatCurveTrack->GetName();
-
-			if(AnimSequenceBase->RawCurveData.GetCurveData(CurveName.UID))
+			const FAnimationCurveIdentifier FloatCurveId(CurveName, ERawCurveTrackTypes::RCT_Float);
+			if (AnimSequenceBase->GetDataModel()->FindCurve(FloatCurveId))
 			{
 				FSmartName TrackName;
 				if (AnimSequenceBase->GetSkeleton()->GetSmartNameByUID(USkeleton::AnimCurveMappingName, CurveName.UID, TrackName))
@@ -349,7 +381,7 @@ void FAnimModel_AnimSequenceBase::RemoveSelectedCurves()
 					IAnimationEditor::FCurveEditInfo EditInfo(Name, Type, CurveEditIndex);
 					OnStopEditingCurves.ExecuteIfBound(TArray<IAnimationEditor::FCurveEditInfo>({ EditInfo }));
 
-					AnimSequenceBase->RawCurveData.DeleteCurveData(TrackName);
+					Controller->RemoveCurve(FloatCurveId);
 					bDeletedCurve = true;
 				}
 			}
@@ -358,10 +390,12 @@ void FAnimModel_AnimSequenceBase::RemoveSelectedCurves()
 		{
 			TSharedRef<FAnimTimelineTrack_TransformCurve> TransformCurveTrack = StaticCastSharedRef<FAnimTimelineTrack_TransformCurve>(SelectedTrack);
 
-			FTransformCurve& TransformCurve = TransformCurveTrack->GetTransformCurve();
+			const FTransformCurve& TransformCurve = TransformCurveTrack->GetTransformCurve();
 			FSmartName CurveName = TransformCurveTrack->GetName();
 
-			if(AnimSequenceBase->RawCurveData.GetCurveData(CurveName.UID, ERawCurveTrackTypes::RCT_Transform))
+			const FAnimationCurveIdentifier TransformCurveId(CurveName, ERawCurveTrackTypes::RCT_Transform);
+
+			if (AnimSequenceBase->GetDataModel()->FindTransformCurve(TransformCurveId))
 			{
 				FSmartName CurveToDelete;
 				if (AnimSequenceBase->GetSkeleton()->GetSmartNameByUID(USkeleton::AnimTrackCurveMappingName, CurveName.UID, CurveToDelete))
@@ -380,17 +414,18 @@ void FAnimModel_AnimSequenceBase::RemoveSelectedCurves()
 
 					OnStopEditingCurves.ExecuteIfBound(CurveEditInfo);
 
-					AnimSequenceBase->RawCurveData.DeleteCurveData(CurveToDelete, ERawCurveTrackTypes::RCT_Transform);
-					
+					Controller->RemoveCurve(TransformCurveId);
 					bDeletedCurve = true;
 				}
 			}	
 		}
 	}
 
+	Controller->CloseBracket();
+
 	if(bDeletedCurve)
 	{
-		AnimSequenceBase->MarkRawDataAsModified();
+		REFACTOR_V2_WARNING("V2 why is post edit change needed here?");
 		AnimSequenceBase->PostEditChange();
 
 		if (GetPreviewScene()->GetPreviewMeshComponent()->PreviewInstance != nullptr)
@@ -398,8 +433,6 @@ void FAnimModel_AnimSequenceBase::RemoveSelectedCurves()
 			GetPreviewScene()->GetPreviewMeshComponent()->PreviewInstance->RefreshCurveBoneControllers();
 		}
 	}
-
-	RefreshTracks();
 }
 
 void FAnimModel_AnimSequenceBase::SetDisplayFormat(EFrameNumberDisplayFormats InFormat)
@@ -435,7 +468,7 @@ bool FAnimModel_AnimSequenceBase::IsDisplaySecondaryChecked() const
 void FAnimModel_AnimSequenceBase::HandleUndoRedo()
 {
 	// Close any curves that are no longer editable
-	for(FFloatCurve& FloatCurve : AnimSequenceBase->RawCurveData.FloatCurves)
+	for (const FFloatCurve& FloatCurve : AnimSequenceBase->GetDataModel()->GetCurveData().FloatCurves)
 	{
 		if(FloatCurve.GetCurveTypeFlag(AACF_Metadata))
 		{
diff --git a/Engine/Source/Editor/Persona/Private/AnimTimeline/AnimModel_AnimSequenceBase.h b/Engine/Source/Editor/Persona/Private/AnimTimeline/AnimModel_AnimSequenceBase.h
--- a/Engine/Source/Editor/Persona/Private/AnimTimeline/AnimModel_AnimSequenceBase.h
+++ b/Engine/Source/Editor/Persona/Private/AnimTimeline/AnimModel_AnimSequenceBase.h
@@ -6,6 +6,8 @@
 #include "PersonaDelegates.h"
 #include "SAnimTimingPanel.h"
 #include "EditorUndoClient.h"
+#include "Animation/AnimSequenceHelpers.h"
+#include "Animation/AnimData/AnimDataModelNotifyCollector.h"
 
 class UAnimSequenceBase;
 class FAnimTimelineTrack_Notifies;
@@ -60,6 +62,9 @@ protected:
 	/** Refresh curve tracks */
 	void RefreshCurveTracks();
 
+	/** Callback for any change made to the UAnimDataModel embedded in the AnimSequenceBase instance this represents */
+	void OnDataModelChanged(const EAnimDataModelNotifyType& NotifyType, UAnimDataModel* Model, const FAnimDataModelNotifPayload& PayLoad);
+
 private:
 	/** UI handlers */
 	void EditSelectedCurves();
@@ -91,4 +96,6 @@ private:
 
 	/** Display flags for notifies track */
 	bool NotifiesTimingElementNodeDisplayFlags[ETimingElementType::Max];
+
+	UE::Anim::FAnimDataModelNotifyCollector NotifyCollector;
 };
\ No newline at end of file
diff --git a/Engine/Source/Editor/Persona/Private/AnimTimeline/AnimTimelineTrack_Curve.cpp b/Engine/Source/Editor/Persona/Private/AnimTimeline/AnimTimelineTrack_Curve.cpp
--- a/Engine/Source/Editor/Persona/Private/AnimTimeline/AnimTimelineTrack_Curve.cpp
+++ b/Engine/Source/Editor/Persona/Private/AnimTimeline/AnimTimelineTrack_Curve.cpp
@@ -159,7 +159,7 @@ FAnimTimelineTrack_Curve::FAnimTimelineTrack_Curve(const FText& InCurveName, con
 	SetHeight(32.0f);
 }
 
-FAnimTimelineTrack_Curve::FAnimTimelineTrack_Curve(FRichCurve& InCurve, const FSmartName& InName, int32 InCurveIndex, ERawCurveTrackTypes InType, const FText& InCurveName, const FText& InFullCurveName, const FLinearColor& InColor, const FLinearColor& InBackgroundColor, const TSharedRef<FAnimModel>& InModel)
+FAnimTimelineTrack_Curve::FAnimTimelineTrack_Curve(const FRichCurve* InCurve, const FSmartName& InName, int32 InCurveIndex, ERawCurveTrackTypes InType, const FText& InCurveName, const FText& InFullCurveName, const FLinearColor& InColor, const FLinearColor& InBackgroundColor, const TSharedRef<FAnimModel>& InModel)
 	: FAnimTimelineTrack(InCurveName, InCurveName, InModel)
 	, Color(InColor)
 	, BackgroundColor(InBackgroundColor)
@@ -168,11 +168,11 @@ FAnimTimelineTrack_Curve::FAnimTimelineTrack_Curve(FRichCurve& InCurve, const FS
 	, OuterCurveIndex(InCurveIndex)
 	, OuterType(InType)
 {
-	Curves.Add(&InCurve);
+	Curves.Add(InCurve);
 	SetHeight(32.0f);
 }
 
-FAnimTimelineTrack_Curve::FAnimTimelineTrack_Curve(const TArray<FRichCurve*>& InCurves, const FText& InCurveName, const FText& InFullCurveName, const FLinearColor& InColor, const FLinearColor& InBackgroundColor, const TSharedRef<FAnimModel>& InModel)
+FAnimTimelineTrack_Curve::FAnimTimelineTrack_Curve(const TArray<const FRichCurve*>& InCurves, const FText& InCurveName, const FText& InFullCurveName, const FLinearColor& InColor, const FLinearColor& InBackgroundColor, const TSharedRef<FAnimModel>& InModel)
 	: FAnimTimelineTrack(InCurveName, InCurveName, InModel)
 	, Curves(InCurves)
 	, Color(InColor)
@@ -193,7 +193,7 @@ TSharedRef<SWidget> FAnimTimelineTrack_Curve::GenerateContainerWidgetForTimeline
 
 	for(int32 CurveIndex = 0; CurveIndex < Curves.Num(); ++CurveIndex)
 	{
-		FRichCurve* Curve = Curves[CurveIndex];
+		const FRichCurve* Curve = Curves[CurveIndex];
 
 		FSmartName Name;
 		ERawCurveTrackTypes Type;
@@ -363,8 +363,6 @@ FReply FAnimTimelineTrack_Curve::HandleDoubleClicked(const FGeometry& InGeometry
 void FAnimTimelineTrack_Curve::HandleCurveChanged()
 {
 	ZoomToFit();
-
-	GetModel()->GetAnimSequenceBase()->MarkRawDataAsModified();
 }
 
 void FAnimTimelineTrack_Curve::PostUndoRedo()
diff --git a/Engine/Source/Editor/Persona/Private/AnimTimeline/AnimTimelineTrack_Curve.h b/Engine/Source/Editor/Persona/Private/AnimTimeline/AnimTimelineTrack_Curve.h
--- a/Engine/Source/Editor/Persona/Private/AnimTimeline/AnimTimelineTrack_Curve.h
+++ b/Engine/Source/Editor/Persona/Private/AnimTimeline/AnimTimelineTrack_Curve.h
@@ -5,6 +5,8 @@
 #include "AnimTimelineTrack.h"
 #include "Animation/Skeleton.h"
 #include "EditorUndoClient.h"
+#include "AnimModel.h"
+#include "Animation/AnimData/CurveIdentifier.h"
 
 class SBorder;
 class FCurveEditor;
@@ -18,8 +20,8 @@ class FAnimTimelineTrack_Curve : public FAnimTimelineTrack, public FSelfRegister
 
 public:
 	FAnimTimelineTrack_Curve(const FText& InCurveName, const FText& InFullCurveName, const FLinearColor& InColor, const FLinearColor& InBackgroundColor, const TSharedRef<FAnimModel>& InModel);
-	FAnimTimelineTrack_Curve(FRichCurve& InCurve, const FSmartName& InName, int32 InCurveIndex, ERawCurveTrackTypes InType, const FText& InCurveName, const FText& InFullCurveName, const FLinearColor& InColor, const FLinearColor& InBackgroundColor, const TSharedRef<FAnimModel>& InModel);
-	FAnimTimelineTrack_Curve(const TArray<FRichCurve*>& InCurves, const FText& InCurveName, const FText& InFullCurveName, const FLinearColor& InColor, const FLinearColor& InBackgroundColor, const TSharedRef<FAnimModel>& InModel);
+	FAnimTimelineTrack_Curve(const FRichCurve* InCurve, const FSmartName& InName, int32 InCurveIndex, ERawCurveTrackTypes InType, const FText& InCurveName, const FText& InFullCurveName, const FLinearColor& InColor, const FLinearColor& InBackgroundColor, const TSharedRef<FAnimModel>& InModel);
+	FAnimTimelineTrack_Curve(const TArray<const FRichCurve*>& InCurves, const FText& InCurveName, const FText& InFullCurveName, const FLinearColor& InColor, const FLinearColor& InBackgroundColor, const TSharedRef<FAnimModel>& InModel);
 
 	/** FAnimTimelineTrack interface */
 	virtual TSharedRef<SWidget> GenerateContainerWidgetForTimeline() override;
@@ -32,7 +34,7 @@ public:
 	virtual void PostRedo(bool bSuccess) override { PostUndoRedo(); }
 
 	/** Access the curve we are editing */
-	const TArray<FRichCurve*>& GetCurves() { return Curves; }
+	const TArray<const FRichCurve*>& GetCurves() { return Curves; }
 
 	/** Get a color for the curve widget (and edited curves) */
 	virtual FLinearColor GetCurveColor(int32 InCurveIndex) const;
@@ -79,7 +81,7 @@ protected:
 
 protected:
 	/** The curve we are editing */
-	TArray<FRichCurve*> Curves;
+	TArray<const FRichCurve*> Curves;
 
 	/** The color of the curve */
 	FLinearColor Color;
diff --git a/Engine/Source/Editor/Persona/Private/AnimTimeline/AnimTimelineTrack_Curves.cpp b/Engine/Source/Editor/Persona/Private/AnimTimeline/AnimTimelineTrack_Curves.cpp
--- a/Engine/Source/Editor/Persona/Private/AnimTimeline/AnimTimelineTrack_Curves.cpp
+++ b/Engine/Source/Editor/Persona/Private/AnimTimeline/AnimTimelineTrack_Curves.cpp
@@ -12,6 +12,7 @@
 #include "Animation/AnimMontage.h"
 #include "SAnimOutlinerItem.h"
 #include "Preferences/PersonaOptions.h"
+#include "Animation/AnimData/AnimDataController.h"
 
 #define LOCTEXT_NAMESPACE "FAnimTimelineTrack_Notifies"
 
@@ -53,7 +54,7 @@ TSharedRef<SWidget> FAnimTimelineTrack_Curves::GenerateContainerWidgetForOutline
 			.Text_Lambda([this]()
 			{ 
 				UAnimSequenceBase* AnimSequenceBase = GetModel()->GetAnimSequenceBase();
-				return FText::Format(LOCTEXT("CurveCountFormat", "({0})"), FText::AsNumber(AnimSequenceBase->RawCurveData.FloatCurves.Num())); 
+				return FText::Format(LOCTEXT("CurveCountFormat", "({0})"), FText::AsNumber(AnimSequenceBase->GetDataModel()->GetNumberOfFloatCurves())); 
 			})
 		];
 
@@ -75,14 +76,9 @@ TSharedRef<SWidget> FAnimTimelineTrack_Curves::GenerateContainerWidgetForOutline
 
 void FAnimTimelineTrack_Curves::DeleteAllCurves()
 {
-	const FScopedTransaction Transaction( LOCTEXT("AnimCurve_RemoveAllCurves", "Remove All Curves") );
-
 	UAnimSequenceBase* AnimSequenceBase = GetModel()->GetAnimSequenceBase();
-	AnimSequenceBase->Modify(true);
-	AnimSequenceBase->RawCurveData.DeleteAllCurveData();
-	AnimSequenceBase->MarkRawDataAsModified();
-
-	GetModel()->RefreshTracks();
+	UAnimDataController* Controller = AnimSequenceBase->GetController();
+	Controller->RemoveAllCurvesOfType(ERawCurveTrackTypes::RCT_Float);
 }
 
 TSharedRef<SWidget> FAnimTimelineTrack_Curves::BuildCurvesSubMenu()
@@ -104,7 +100,7 @@ TSharedRef<SWidget> FAnimTimelineTrack_Curves::BuildCurvesSubMenu()
 		);
 
 		UAnimSequenceBase* AnimSequenceBase = GetModel()->GetAnimSequenceBase();
-		if(AnimSequenceBase->RawCurveData.FloatCurves.Num() > 0)
+		if(AnimSequenceBase->GetDataModel()->GetNumberOfFloatCurves() > 0)
 		{
 			MenuBuilder.AddMenuEntry(
 				FAnimSequenceTimelineCommands::Get().RemoveAllCurves->GetLabel(),
@@ -171,15 +167,17 @@ void FAnimTimelineTrack_Curves::FillMetadataEntryMenu(FMenuBuilder& Builder)
 	{
 		TArray<FSmartNameSortItem> SmartNameList;
 
+		const TArray<FFloatCurve>& FloatCurves = AnimSequenceBase->GetDataModel()->GetFloatCurves();
+
 		for (USkeleton::AnimCurveUID Id : CurveUids)
 		{
-			if (!AnimSequenceBase->RawCurveData.GetCurveData(Id))
+			FSmartName SmartName;
+			if (Mapping->FindSmartNameByUID(Id, SmartName) && !FloatCurves.ContainsByPredicate([SmartName](FFloatCurve& Curve)
 			{
-				FName CurveName;
-				if (Mapping->GetName(Id, CurveName))
-				{
-					SmartNameList.Add(FSmartNameSortItem(CurveName, Id));
-				}
+				return Curve.Name == SmartName;
+			}))
+			{
+				SmartNameList.Add(FSmartNameSortItem(SmartName.DisplayName, Id));
 			}
 		}
 
@@ -235,7 +233,7 @@ void FAnimTimelineTrack_Curves::FillVariableCurveMenu(FMenuBuilder& Builder)
 
 		for (USkeleton::AnimCurveUID Id : CurveUids)
 		{
-			if (!AnimSequenceBase->RawCurveData.GetCurveData(Id))
+			if (!AnimSequenceBase->GetDataModel()->FindFloatCurve(FAnimationCurveIdentifier(Id, ERawCurveTrackTypes::RCT_Float)))
 			{
 				FName CurveName;
 				if (Mapping->GetName(Id, CurveName))
@@ -270,19 +268,12 @@ void FAnimTimelineTrack_Curves::AddMetadataEntry(USkeleton::AnimCurveUID Uid)
 	UAnimSequenceBase* AnimSequenceBase = GetModel()->GetAnimSequenceBase();
 	ensureAlways(AnimSequenceBase->GetSkeleton()->GetSmartNameByUID(USkeleton::AnimCurveMappingName, Uid, NewName));
 
-	FScopedTransaction Transaction(LOCTEXT("AddCurveMetadata", "Add Curve Metadata"));
+	UAnimDataController* Controller = AnimSequenceBase->GetController();
+	UAnimDataController::FScopedBracket ScopedBracket(Controller, LOCTEXT("AddCurveMetadata", "Add Curve Metadata"));
 
-	AnimSequenceBase->Modify(true);
-
-	if(AnimSequenceBase->RawCurveData.AddCurveData(NewName))
-	{
-		AnimSequenceBase->MarkRawDataAsModified();
-		FFloatCurve* Curve = static_cast<FFloatCurve *>(AnimSequenceBase->RawCurveData.GetCurveData(Uid, ERawCurveTrackTypes::RCT_Float));
-		Curve->FloatCurve.AddKey(0.0f, 1.0f);
-		Curve->SetCurveTypeFlag(AACF_Metadata, true);
-		
-		GetModel()->RefreshTracks();
-	}
+	const FAnimationCurveIdentifier MetadataCurveId(NewName, ERawCurveTrackTypes::RCT_Float);
+	Controller->AddCurve(MetadataCurveId, AACF_Metadata);
+	Controller->SetCurveKeys(MetadataCurveId, { FRichCurveKey(0.f, 1.f) });	
 }
 
 void FAnimTimelineTrack_Curves::CreateNewMetadataEntryClicked()
@@ -370,10 +361,10 @@ void FAnimTimelineTrack_Curves::AddVariableCurve(USkeleton::AnimCurveUID CurveUi
 	USkeleton* Skeleton = AnimSequenceBase->GetSkeleton();
 	FSmartName NewName;
 	ensureAlways(Skeleton->GetSmartNameByUID(USkeleton::AnimCurveMappingName, CurveUid, NewName));
-	AnimSequenceBase->RawCurveData.AddCurveData(NewName);
-	AnimSequenceBase->MarkRawDataAsModified();
 
-	GetModel()->RefreshTracks();
+	UAnimDataController* Controller = AnimSequenceBase->GetController();
+	const FAnimationCurveIdentifier FloatCurveId(NewName, ERawCurveTrackTypes::RCT_Float);
+	Controller->AddCurve(FloatCurveId);
 }
 
 void FAnimTimelineTrack_Curves::HandleShowCurvePoints()
