Moved SmartName FName arrays local to functions, global static is being initialized to NAME_None due to (new) different AnimModifiers module loading order

diff --git a/Engine/Source/Editor/AnimationModifiers/Private/AnimationBlueprintLibrary.cpp b/Engine/Source/Editor/AnimationModifiers/Private/AnimationBlueprintLibrary.cpp
--- a/Engine/Source/Editor/AnimationModifiers/Private/AnimationBlueprintLibrary.cpp
+++ b/Engine/Source/Editor/AnimationModifiers/Private/AnimationBlueprintLibrary.cpp
@@ -23,8 +23,6 @@
 
 DEFINE_LOG_CATEGORY_STATIC(LogAnimationBlueprintLibrary, Verbose, All);
 
-const FName UAnimationBlueprintLibrary::SmartContainerNames[(int32)ESmartNameContainerType::SNCT_MAX] = { USkeleton::AnimCurveMappingName, USkeleton::AnimTrackCurveMappingName };
-
 void UAnimationBlueprintLibrary::GetNumFrames(const UAnimSequence* AnimationSequence, int32& NumFrames)
 {
 	NumFrames = 0;
@@ -1197,7 +1195,8 @@ void UAnimationBlueprintLibrary::AddCurve(UAnimSequence* AnimationSequence, FNam
 		if (bValidMetaData && bValidTransformCurveData )
 		{
 			// Add or retrieve the smartname
-			const bool bCurveAdded = AddCurveInternal(AnimationSequence, CurveName, SmartContainerNames[(int32)CurveContainer], CurveFlags, CurveType);
+			const FName Names[(int32)ESmartNameContainerType::SNCT_MAX] = { USkeleton::AnimCurveMappingName, USkeleton::AnimTrackCurveMappingName };
+			const bool bCurveAdded = AddCurveInternal(AnimationSequence, CurveName, Names[(int32)CurveContainer], CurveFlags, CurveType);
 
 			if (!bCurveAdded)
 			{
@@ -1648,12 +1647,13 @@ FSmartName UAnimationBlueprintLibrary::RetrieveSmartNameForCurve(const UAnimSequ
 FName UAnimationBlueprintLibrary::RetrieveContainerNameForCurve(const UAnimSequence* AnimationSequence, FName CurveName)
 {
 	checkf(AnimationSequence != nullptr, TEXT("Invalid Animation Sequence ptr"));
+	const FName Names[(int32)ESmartNameContainerType::SNCT_MAX] = { USkeleton::AnimCurveMappingName, USkeleton::AnimTrackCurveMappingName };
 	for (int32 Index = 0; Index < (int32)ESmartNameContainerType::SNCT_MAX; ++Index)
 	{
-		const FSmartNameMapping* CurveMapping = AnimationSequence->GetSkeleton()->GetSmartNameContainer(SmartContainerNames[Index]);
+		const FSmartNameMapping* CurveMapping = AnimationSequence->GetSkeleton()->GetSmartNameContainer(Names[Index]);
 		if (CurveMapping && CurveMapping->Exists(CurveName))
 		{
-			return SmartContainerNames[Index];
+			return Names[Index];
 		}
 	}
 
diff --git a/Engine/Source/Editor/AnimationModifiers/Public/AnimationBlueprintLibrary.h b/Engine/Source/Editor/AnimationModifiers/Public/AnimationBlueprintLibrary.h
--- a/Engine/Source/Editor/AnimationModifiers/Public/AnimationBlueprintLibrary.h
+++ b/Engine/Source/Editor/AnimationModifiers/Public/AnimationBlueprintLibrary.h
@@ -466,6 +466,4 @@ public:
 	/** Finds the Bone Path from the given Bone to the Root Bone */
 	UFUNCTION(BlueprintPure, Category = "AnimationBlueprintLibrary|Helpers")
 	static void FindBonePathToRoot(const UAnimSequence* AnimationSequence, FName BoneName, TArray<FName>& BonePath);
-
-	static const FName SmartContainerNames[(int32)ESmartNameContainerType::SNCT_MAX];
 };
