Expose Montages notifies to python via anim blueprint library like anim sequences already are
#fix changed subset of the AnimationBlueprintLibrary API to take an UAnimSequenceBase/UAnimationAsset instead of an UAnimSequence
#misc added redirectors for all cases

diff --git a/Engine/Config/BaseEngine.ini b/Engine/Config/BaseEngine.ini
--- a/Engine/Config/BaseEngine.ini
+++ b/Engine/Config/BaseEngine.ini
@@ -868,6 +868,51 @@ NearClipPlane=10.0
 +ClassRedirects=(OldName="/Script/MeshDescription.MeshDescriptionBulkDataWrapper",NewName="/Script/Engine.StaticMeshDescriptionBulkData")
 +StructRedirects=(OldName="/Script/Engine.FilePath",NewName="/Script/CoreUObject.FilePath")
 +StructRedirects=(OldName="/Script/Engine.DirectoryPath",NewName="/Script/CoreUObject.DirectoryPath")
++PropertyRedirects=(OldName="AnimationBlueprintLibrary.GetNumFrames.AnimationSequence",NewName="AnimationSequenceBase")
++PropertyRedirects=(OldName="AnimationBlueprintLibrary.GetNumKeys.AnimationSequence",NewName="AnimationSequenceBase")
++PropertyRedirects=(OldName="AnimationBlueprintLibrary.GetAnimationTrackNames.AnimationSequence",NewName="AnimationSequenceBase")
++PropertyRedirects=(OldName="AnimationBlueprintLibrary.GetRawTrackPositionData.AnimationSequence",NewName="AnimationSequenceBase")
++PropertyRedirects=(OldName="AnimationBlueprintLibrary.GetRawTrackRotationData.AnimationSequence",NewName="AnimationSequenceBase")
++PropertyRedirects=(OldName="AnimationBlueprintLibrary.GetRawTrackScaleData.AnimationSequence",NewName="AnimationSequenceBase")
++PropertyRedirects=(OldName="AnimationBlueprintLibrary.GetRawTrackData.AnimationSequence",NewName="AnimationSequenceBase")
++PropertyRedirects=(OldName="AnimationBlueprintLibrary.IsValidRawAnimationTrackName.AnimationSequence",NewName="AnimationSequenceBase")
++PropertyRedirects=(OldName="AnimationBlueprintLibrary.GetAnimationNotifyEvents.AnimationSequence",NewName="AnimationSequenceBase")
++PropertyRedirects=(OldName="AnimationBlueprintLibrary.GetAnimationNotifyEventNames.AnimationSequence",NewName="AnimationSequenceBase")
++PropertyRedirects=(OldName="AnimationBlueprintLibrary.AddAnimationNotifyEvent.AnimationSequence",NewName="AnimationSequenceBase")
++PropertyRedirects=(OldName="AnimationBlueprintLibrary.AddAnimationNotifyStateEvent.AnimationSequence",NewName="AnimationSequenceBase")
++PropertyRedirects=(OldName="AnimationBlueprintLibrary.AddAnimationNotifyEventObject.AnimationSequence",NewName="AnimationSequenceBase")
++PropertyRedirects=(OldName="AnimationBlueprintLibrary.AddAnimationNotifyStateEventObject.AnimationSequence",NewName="AnimationSequenceBase")
++PropertyRedirects=(OldName="AnimationBlueprintLibrary.RemoveAnimationNotifyEventsByName.AnimationSequence",NewName="AnimationSequenceBase")
++PropertyRedirects=(OldName="AnimationBlueprintLibrary.RemoveAnimationNotifyEventsByTrack.AnimationSequence",NewName="AnimationSequenceBase")
++PropertyRedirects=(OldName="AnimationBlueprintLibrary.ReplaceAnimNotifyStates.AnimationSequence",NewName="AnimationSequenceBase")
++PropertyRedirects=(OldName="AnimationBlueprintLibrary.ReplaceAnimNotifies.AnimationSequence",NewName="AnimationSequenceBase")
++PropertyRedirects=(OldName="AnimationBlueprintLibrary.CopyAnimNotifiesFromSequence.SrcAnimSequence",NewName="SourceAnimationSequenceBase")
++PropertyRedirects=(OldName="AnimationBlueprintLibrary.CopyAnimNotifiesFromSequence.DestAnimSequence",NewName="DestinationAnimationSequenceBase")
++PropertyRedirects=(OldName="AnimationBlueprintLibrary.GetAnimationNotifyTrackNames.AnimationSequence",NewName="AnimationSequenceBase")
++PropertyRedirects=(OldName="AnimationBlueprintLibrary.AddAnimationNotifyTrack.AnimationSequence",NewName="AnimationSequenceBase")
++PropertyRedirects=(OldName="AnimationBlueprintLibrary.RemoveAnimationNotifyTrack.AnimationSequence",NewName="AnimationSequenceBase")
++PropertyRedirects=(OldName="AnimationBlueprintLibrary.RemoveAllAnimationNotifyTracks.AnimationSequence",NewName="AnimationSequenceBase")
++PropertyRedirects=(OldName="AnimationBlueprintLibrary.IsValidAnimNotifyTrackName.AnimationSequence",NewName="AnimationSequenceBase")
++PropertyRedirects=(OldName="AnimationBlueprintLibrary.GetAnimationNotifyEventsForTrack.AnimationSequence",NewName="AnimationSequenceBase")
++PropertyRedirects=(OldName="AnimationBlueprintLibrary.AddMetaData.AnimationSequence",NewName="AnimationAsset")
++PropertyRedirects=(OldName="AnimationBlueprintLibrary.AddMetaDataObject.AnimationSequence",NewName="AnimationAsset")
++PropertyRedirects=(OldName="AnimationBlueprintLibrary.RemoveAllMetaData.AnimationSequence",NewName="AnimationAsset")
++PropertyRedirects=(OldName="AnimationBlueprintLibrary.RemoveMetaData.AnimationSequence",NewName="AnimationAsset")
++PropertyRedirects=(OldName="AnimationBlueprintLibrary.RemoveMetaDataOfClass.AnimationSequence",NewName="AnimationAsset")
++PropertyRedirects=(OldName="AnimationBlueprintLibrary.GetMetaData.AnimationSequence",NewName="AnimationAsset")
++PropertyRedirects=(OldName="AnimationBlueprintLibrary.GetMetaDataOfClass.AnimationSequence",NewName="AnimationAsset")
++PropertyRedirects=(OldName="AnimationBlueprintLibrary.ContainsMetaDataOfClass.AnimationSequence",NewName="AnimationAsset")
++PropertyRedirects=(OldName="AnimationBlueprintLibrary.GetBonePoseForTime.AnimationSequence",NewName="AnimationSequenceBase")
++PropertyRedirects=(OldName="AnimationBlueprintLibrary.GetBonePoseForFrame.AnimationSequence",NewName="AnimationSequenceBase")
++PropertyRedirects=(OldName="AnimationBlueprintLibrary.GetBonePosesForTime.AnimationSequence",NewName="AnimationSequenceBase")
++PropertyRedirects=(OldName="AnimationBlueprintLibrary.GetBonePosesForFrame.AnimationSequence",NewName="AnimationSequenceBase")
++PropertyRedirects=(OldName="AnimationBlueprintLibrary.GetSequenceLength.AnimationSequence",NewName="AnimationSequenceBase")
++PropertyRedirects=(OldName="AnimationBlueprintLibrary.GetRateScale.AnimationSequence",NewName="AnimationSequenceBase")
++PropertyRedirects=(OldName="AnimationBlueprintLibrary.SetRateScale.AnimationSequence",NewName="AnimationSequenceBase")
++PropertyRedirects=(OldName="AnimationBlueprintLibrary.GetFrameAtTime.AnimationSequence",NewName="AnimationSequenceBase")
++PropertyRedirects=(OldName="AnimationBlueprintLibrary.GetTimeAtFrame.AnimationSequence",NewName="AnimationSequenceBase")
++PropertyRedirects=(OldName="AnimationBlueprintLibrary.IsValidTime.AnimationSequence",NewName="AnimationSequenceBase")
++PropertyRedirects=(OldName="AnimationBlueprintLibrary.FindBonePathToRoot.AnimationSequence",NewName="AnimationSequenceBase")
 
 [CoreUObject.Metadata]
 ; Note: UnrealHeaderTool should be rerun after making changes to MetadataRedirects to catch any keys specified in class headers
diff --git a/Engine/Source/Editor/AnimationModifiers/Private/AnimationBlueprintLibrary.cpp b/Engine/Source/Editor/AnimationModifiers/Private/AnimationBlueprintLibrary.cpp
--- a/Engine/Source/Editor/AnimationModifiers/Private/AnimationBlueprintLibrary.cpp
+++ b/Engine/Source/Editor/AnimationModifiers/Private/AnimationBlueprintLibrary.cpp
@@ -23,12 +23,12 @@
 
 DEFINE_LOG_CATEGORY_STATIC(LogAnimationBlueprintLibrary, Verbose, All);
 
-void UAnimationBlueprintLibrary::GetNumFrames(const UAnimSequence* AnimationSequence, int32& NumFrames)
+void UAnimationBlueprintLibrary::GetNumFrames(const UAnimSequenceBase* AnimationSequenceBase, int32& NumFrames)
 {
 	NumFrames = 0;
-	if (AnimationSequence)
+	if (AnimationSequenceBase)
 	{
-		NumFrames = AnimationSequence->GetNumberOfSampledKeys() - 1;
+		NumFrames = AnimationSequenceBase->GetNumberOfSampledKeys() - 1;
 	}
 	else
 	{
@@ -36,12 +36,12 @@ void UAnimationBlueprintLibrary::GetNumFrames(const UAnimSequence* AnimationSequ
 	}
 }
 
-void UAnimationBlueprintLibrary::GetNumKeys(const UAnimSequence* AnimationSequence, int32& NumKeys)
+void UAnimationBlueprintLibrary::GetNumKeys(const UAnimSequenceBase* AnimationSequenceBase, int32& NumKeys)
 {
 	NumKeys = 0;
-	if (AnimationSequence)
+	if (AnimationSequenceBase)
 	{
-		NumKeys = AnimationSequence->GetNumberOfSampledKeys();
+		NumKeys = AnimationSequenceBase->GetNumberOfSampledKeys();
 	}
 	else
 	{
@@ -49,12 +49,12 @@ void UAnimationBlueprintLibrary::GetNumKeys(const UAnimSequence* AnimationSequen
 	}
 }
 
-void UAnimationBlueprintLibrary::GetAnimationTrackNames(const UAnimSequence* AnimationSequence, TArray<FName>& TrackNames)
+void UAnimationBlueprintLibrary::GetAnimationTrackNames(const UAnimSequenceBase* AnimationSequenceBase, TArray<FName>& TrackNames)
 {
 	TrackNames.Empty();
-	if (AnimationSequence)
+	if (AnimationSequenceBase)
 	{
-		AnimationSequence->GetDataModel()->GetBoneTrackNames(TrackNames);
+		AnimationSequenceBase->GetDataModel()->GetBoneTrackNames(TrackNames);
 	}
 	else
 	{
@@ -99,77 +99,77 @@ void UAnimationBlueprintLibrary::GetAnimationCurveNames(const UAnimSequence* Ani
 	}
 }
 
-void UAnimationBlueprintLibrary::GetRawTrackPositionData(const UAnimSequence* AnimationSequence, const FName TrackName, TArray<FVector>& PositionData)
+void UAnimationBlueprintLibrary::GetRawTrackPositionData(const UAnimSequenceBase* AnimationSequenceBase, const FName TrackName, TArray<FVector>& PositionData)
 {
 	PositionData.Empty();
-	if (IsValidRawAnimationTrackName(AnimationSequence, TrackName))
+	if (IsValidRawAnimationTrackName(AnimationSequenceBase, TrackName))
 	{
-		const FRawAnimSequenceTrack& RawTrack =	GetRawAnimationTrackByName(AnimationSequence, TrackName);
+		const FRawAnimSequenceTrack& RawTrack =	GetRawAnimationTrackByName(AnimationSequenceBase, TrackName);
 		PositionData.Append(RawTrack.PosKeys);
 	}
 	else
 	{
-		const FString AnimSequenceName = (AnimationSequence != nullptr) ? AnimationSequence->GetName() : "Invalid Animation sequence";
+		const FString AnimSequenceName = (AnimationSequenceBase != nullptr) ? AnimationSequenceBase->GetName() : "Invalid Animation sequence";
 		UE_LOG(LogAnimationBlueprintLibrary, Warning, TEXT("Raw Animation Track name %s does not exist in Animation Sequence %s"), *TrackName.ToString(), *AnimSequenceName );
 	}	
 }
 
-void UAnimationBlueprintLibrary::GetRawTrackRotationData(const UAnimSequence* AnimationSequence, const FName TrackName, TArray<FQuat>& RotationData)
+void UAnimationBlueprintLibrary::GetRawTrackRotationData(const UAnimSequenceBase* AnimationSequenceBase, const FName TrackName, TArray<FQuat>& RotationData)
 {
 	RotationData.Empty();
-	if (IsValidRawAnimationTrackName(AnimationSequence, TrackName))
+	if (IsValidRawAnimationTrackName(AnimationSequenceBase, TrackName))
 	{
-		const FRawAnimSequenceTrack& RawTrack = GetRawAnimationTrackByName(AnimationSequence, TrackName);
+		const FRawAnimSequenceTrack& RawTrack = GetRawAnimationTrackByName(AnimationSequenceBase, TrackName);
 		RotationData.Append(RawTrack.RotKeys);
 	}
 	else
 	{	
-		const FString AnimSequenceName = (AnimationSequence != nullptr) ? AnimationSequence->GetName() : "Invalid Animation sequence";
+		const FString AnimSequenceName = (AnimationSequenceBase != nullptr) ? AnimationSequenceBase->GetName() : "Invalid Animation sequence";
 		UE_LOG(LogAnimationBlueprintLibrary, Warning, TEXT("Raw Animation Track name %s does not exist in Animation Sequence %s"), *TrackName.ToString(), *AnimSequenceName);
 	}
 }
 
-void UAnimationBlueprintLibrary::GetRawTrackScaleData(const UAnimSequence* AnimationSequence, const FName TrackName, TArray<FVector>& ScaleData)
+void UAnimationBlueprintLibrary::GetRawTrackScaleData(const UAnimSequenceBase* AnimationSequenceBase, const FName TrackName, TArray<FVector>& ScaleData)
 {
 	ScaleData.Empty();
-	if (IsValidRawAnimationTrackName(AnimationSequence, TrackName))
+	if (IsValidRawAnimationTrackName(AnimationSequenceBase, TrackName))
 	{
-		const FRawAnimSequenceTrack& RawTrack = GetRawAnimationTrackByName(AnimationSequence, TrackName);
+		const FRawAnimSequenceTrack& RawTrack = GetRawAnimationTrackByName(AnimationSequenceBase, TrackName);
 		ScaleData.Append(RawTrack.ScaleKeys);
 	}
 	else
 	{
-		const FString AnimSequenceName = (AnimationSequence != nullptr) ? AnimationSequence->GetName() : "Invalid Animation sequence";
+		const FString AnimSequenceName = (AnimationSequenceBase != nullptr) ? AnimationSequenceBase->GetName() : "Invalid Animation sequence";
 		UE_LOG(LogAnimationBlueprintLibrary, Warning, TEXT("Raw Animation Track name %s does not exist in Animation Sequence %s"), *TrackName.ToString(), *AnimSequenceName);
 	}
 }
 
-void UAnimationBlueprintLibrary::GetRawTrackData(const UAnimSequence* AnimationSequence, const FName TrackName, TArray<FVector>& PositionKeys, TArray<FQuat>& RotationKeys, TArray<FVector>& ScalingKeys)
+void UAnimationBlueprintLibrary::GetRawTrackData(const UAnimSequenceBase* AnimationSequenceBase, const FName TrackName, TArray<FVector>& PositionKeys, TArray<FQuat>& RotationKeys, TArray<FVector>& ScalingKeys)
 {
 	PositionKeys.Empty();
 	RotationKeys.Empty();
 	ScalingKeys.Empty();
-	if (IsValidRawAnimationTrackName(AnimationSequence, TrackName))
+	if (IsValidRawAnimationTrackName(AnimationSequenceBase, TrackName))
 	{		
-		const FRawAnimSequenceTrack& RawTrack = GetRawAnimationTrackByName(AnimationSequence, TrackName);
+		const FRawAnimSequenceTrack& RawTrack = GetRawAnimationTrackByName(AnimationSequenceBase, TrackName);
 		PositionKeys.Append(RawTrack.PosKeys);
 		RotationKeys.Append(RawTrack.RotKeys);
 		ScalingKeys.Append(RawTrack.ScaleKeys);
 	}
 	else
 	{
-		const FString AnimSequenceName = (AnimationSequence != nullptr) ? AnimationSequence->GetName() : "Invalid Animation sequence";
+		const FString AnimSequenceName = (AnimationSequenceBase != nullptr) ? AnimationSequenceBase->GetName() : "Invalid Animation sequence";
 		UE_LOG(LogAnimationBlueprintLibrary, Warning, TEXT("Raw Animation Track name %s does not exist in Animation Sequence %s"), *TrackName.ToString(), *AnimSequenceName);
 	}
 }
 
-bool UAnimationBlueprintLibrary::IsValidRawAnimationTrackName(const UAnimSequence* AnimationSequence, const FName TrackName)
+bool UAnimationBlueprintLibrary::IsValidRawAnimationTrackName(const UAnimSequenceBase* AnimationSequenceBase, const FName TrackName)
 {
 	bool bValidName = false;
 
-	if (AnimationSequence)
+	if (AnimationSequenceBase)
 	{
-		const int32 TrackIndex = AnimationSequence->GetDataModel()->GetBoneTrackIndexByName(TrackName);
+		const int32 TrackIndex = AnimationSequenceBase->GetDataModel()->GetBoneTrackIndexByName(TrackName);
 		bValidName = (TrackIndex != INDEX_NONE);
 	}
 	else
@@ -180,11 +180,11 @@ bool UAnimationBlueprintLibrary::IsValidRawAnimationTrackName(const UAnimSequenc
 	return bValidName;
 }
 
-const FRawAnimSequenceTrack&  UAnimationBlueprintLibrary::GetRawAnimationTrackByName(const UAnimSequence* AnimationSequence, const FName TrackName)
+const FRawAnimSequenceTrack&  UAnimationBlueprintLibrary::GetRawAnimationTrackByName(const UAnimSequenceBase* AnimationSequenceBase, const FName TrackName)
 {
-	checkf(AnimationSequence, TEXT("Invalid Animation Sequence supplied for GetRawAnimationTrackByName"));
+	checkf(AnimationSequenceBase, TEXT("Invalid Animation Sequence supplied for GetRawAnimationTrackByName"));
 
-	const FBoneAnimationTrack& AnimationTrack = AnimationSequence->GetDataModel()->GetBoneTrackByName(TrackName);
+	const FBoneAnimationTrack& AnimationTrack = AnimationSequenceBase->GetDataModel()->GetBoneTrackByName(TrackName);
 	return AnimationTrack.InternalTrackData;
 }
 
@@ -549,12 +549,12 @@ void UAnimationBlueprintLibrary::RemoveAllAnimationSyncMarkers(UAnimSequence* An
 	}
 }
 
-void UAnimationBlueprintLibrary::GetAnimationNotifyEvents(const UAnimSequence* AnimationSequence, TArray<FAnimNotifyEvent>& NotifyEvents)
+void UAnimationBlueprintLibrary::GetAnimationNotifyEvents(const UAnimSequenceBase* AnimationSequenceBase, TArray<FAnimNotifyEvent>& NotifyEvents)
 {
 	NotifyEvents.Empty();
-	if (AnimationSequence)
+	if (AnimationSequenceBase)
 	{
-		NotifyEvents = AnimationSequence->Notifies;
+		NotifyEvents = AnimationSequenceBase->Notifies;
 	}
 	else
 	{
@@ -562,12 +562,12 @@ void UAnimationBlueprintLibrary::GetAnimationNotifyEvents(const UAnimSequence* A
 	}	
 }
 
-void UAnimationBlueprintLibrary::GetAnimationNotifyEventNames(const UAnimSequence* AnimationSequence, TArray<FName>& EventNames)
+void UAnimationBlueprintLibrary::GetAnimationNotifyEventNames(const UAnimSequenceBase* AnimationSequenceBase, TArray<FName>& EventNames)
 {
 	EventNames.Empty();
-	if (AnimationSequence)
+	if (AnimationSequenceBase)
 	{
-		for (const FAnimNotifyEvent& Event : AnimationSequence->Notifies)
+		for (const FAnimNotifyEvent& Event : AnimationSequenceBase->Notifies)
 		{
 			EventNames.AddUnique(Event.NotifyName);
 		}
@@ -578,27 +578,27 @@ void UAnimationBlueprintLibrary::GetAnimationNotifyEventNames(const UAnimSequenc
 	}	
 }
 
-UAnimNotify* UAnimationBlueprintLibrary::AddAnimationNotifyEvent(UAnimSequence* AnimationSequence, FName NotifyTrackName, float StartTime, TSubclassOf<UAnimNotify> NotifyClass)
+UAnimNotify* UAnimationBlueprintLibrary::AddAnimationNotifyEvent(UAnimSequenceBase* AnimationSequenceBase, FName NotifyTrackName, float StartTime, TSubclassOf<UAnimNotify> NotifyClass)
 {
 	UAnimNotify* Notify = nullptr;
-	if (AnimationSequence)
+	if (AnimationSequenceBase)
 	{
-		const bool bIsValidTrackName = IsValidAnimNotifyTrackName(AnimationSequence, NotifyTrackName);
-		const bool bIsValidTime = IsValidTimeInternal(AnimationSequence, StartTime);
+		const bool bIsValidTrackName = IsValidAnimNotifyTrackName(AnimationSequenceBase, NotifyTrackName);
+		const bool bIsValidTime = IsValidTimeInternal(AnimationSequenceBase, StartTime);
 
 		if (bIsValidTrackName && bIsValidTime)
 		{
-			FAnimNotifyEvent& NewEvent = AnimationSequence->Notifies.AddDefaulted_GetRef();
+			FAnimNotifyEvent& NewEvent = AnimationSequenceBase->Notifies.AddDefaulted_GetRef();
 
 			NewEvent.NotifyName = NAME_None;
-			NewEvent.Link(AnimationSequence, StartTime);
-			NewEvent.TriggerTimeOffset = GetTriggerTimeOffsetForType(AnimationSequence->CalculateOffsetForNotify(StartTime));
-			NewEvent.TrackIndex = GetTrackIndexForAnimationNotifyTrackName(AnimationSequence, NotifyTrackName);
+			NewEvent.Link(AnimationSequenceBase, StartTime);
+			NewEvent.TriggerTimeOffset = GetTriggerTimeOffsetForType(AnimationSequenceBase->CalculateOffsetForNotify(StartTime));
+			NewEvent.TrackIndex = GetTrackIndexForAnimationNotifyTrackName(AnimationSequenceBase, NotifyTrackName);
 			NewEvent.NotifyStateClass = nullptr;
 
 			if (NotifyClass)
 			{
-				Notify = NewObject<UAnimNotify>(AnimationSequence, NotifyClass, NAME_None, RF_Transactional);
+				Notify = NewObject<UAnimNotify>(AnimationSequenceBase, NotifyClass, NAME_None, RF_Transactional);
 				NewEvent.Notify = Notify;
 
 				// Setup name for new event
@@ -613,18 +613,18 @@ UAnimNotify* UAnimationBlueprintLibrary::AddAnimationNotifyEvent(UAnimSequence*
 			}
 
 			// Refresh all cached data
-			AnimationSequence->RefreshCacheData();
+			AnimationSequenceBase->RefreshCacheData();
 		}
 		else
 		{
 			if (!bIsValidTrackName)
 			{
-				UE_LOG(LogAnimationBlueprintLibrary, Warning, TEXT("Animation Notify Track %s does not exist on Animation Sequence %s"), *NotifyTrackName.ToString(), *AnimationSequence->GetName());
+				UE_LOG(LogAnimationBlueprintLibrary, Warning, TEXT("Animation Notify Track %s does not exist on Animation Sequence %s"), *NotifyTrackName.ToString(), *AnimationSequenceBase->GetName());
 			}
 
 			if (!bIsValidTime)
 			{
-				UE_LOG(LogAnimationBlueprintLibrary, Warning, TEXT("%f is side of Animation Sequence %s range"), StartTime, *AnimationSequence->GetName());
+				UE_LOG(LogAnimationBlueprintLibrary, Warning, TEXT("%f is side of Animation Sequence %s range"), StartTime, *AnimationSequenceBase->GetName());
 			}
 		}
 	}
@@ -636,27 +636,27 @@ UAnimNotify* UAnimationBlueprintLibrary::AddAnimationNotifyEvent(UAnimSequence*
 	return Notify;
 }
 
-UAnimNotifyState* UAnimationBlueprintLibrary::AddAnimationNotifyStateEvent(UAnimSequence* AnimationSequence, FName NotifyTrackName, float StartTime, float Duration, TSubclassOf<UAnimNotifyState> NotifyStateClass)
+UAnimNotifyState* UAnimationBlueprintLibrary::AddAnimationNotifyStateEvent(UAnimSequenceBase* AnimationSequenceBase, FName NotifyTrackName, float StartTime, float Duration, TSubclassOf<UAnimNotifyState> NotifyStateClass)
 {
 	UAnimNotifyState* NotifyState = nullptr;
-	if (AnimationSequence)
+	if (AnimationSequenceBase)
 	{
-		const bool bIsValidTrackName = IsValidAnimNotifyTrackName(AnimationSequence, NotifyTrackName);
-		const bool bIsValidTime = IsValidTimeInternal(AnimationSequence, StartTime);
+		const bool bIsValidTrackName = IsValidAnimNotifyTrackName(AnimationSequenceBase, NotifyTrackName);
+		const bool bIsValidTime = IsValidTimeInternal(AnimationSequenceBase, StartTime);
 
 		if (bIsValidTrackName && bIsValidTime)
 		{
-			FAnimNotifyEvent& NewEvent = AnimationSequence->Notifies.AddDefaulted_GetRef();
+			FAnimNotifyEvent& NewEvent = AnimationSequenceBase->Notifies.AddDefaulted_GetRef();
 
 			NewEvent.NotifyName = NAME_None;
-			NewEvent.Link(AnimationSequence, StartTime);
-			NewEvent.TriggerTimeOffset = GetTriggerTimeOffsetForType(AnimationSequence->CalculateOffsetForNotify(StartTime));
-			NewEvent.TrackIndex = GetTrackIndexForAnimationNotifyTrackName(AnimationSequence, NotifyTrackName);
+			NewEvent.Link(AnimationSequenceBase, StartTime);
+			NewEvent.TriggerTimeOffset = GetTriggerTimeOffsetForType(AnimationSequenceBase->CalculateOffsetForNotify(StartTime));
+			NewEvent.TrackIndex = GetTrackIndexForAnimationNotifyTrackName(AnimationSequenceBase, NotifyTrackName);
 			NewEvent.Notify = nullptr;
 
 			if (NotifyStateClass)
 			{
-				NotifyState = NewObject<UAnimNotifyState>(AnimationSequence, NotifyStateClass, NAME_None, RF_Transactional);
+				NotifyState = NewObject<UAnimNotifyState>(AnimationSequenceBase, NotifyStateClass, NAME_None, RF_Transactional);
 				NewEvent.NotifyStateClass = NotifyState;
 
 				// Setup name and duration for new event
@@ -664,7 +664,7 @@ UAnimNotifyState* UAnimationBlueprintLibrary::AddAnimationNotifyStateEvent(UAnim
 				{
 					NewEvent.NotifyName = FName(*NewEvent.NotifyStateClass->GetNotifyName());
 					NewEvent.SetDuration(Duration);
-					NewEvent.EndLink.Link(AnimationSequence, NewEvent.EndLink.GetTime());
+					NewEvent.EndLink.Link(AnimationSequenceBase, NewEvent.EndLink.GetTime());
 				}
 			}
 			else
@@ -673,18 +673,18 @@ UAnimNotifyState* UAnimationBlueprintLibrary::AddAnimationNotifyStateEvent(UAnim
 			}
 
 			// Refresh all cached data
-			AnimationSequence->RefreshCacheData();			
+			AnimationSequenceBase->RefreshCacheData();			
 		}
 		else
 		{
 			if (!bIsValidTrackName)
 			{
-				UE_LOG(LogAnimationBlueprintLibrary, Warning, TEXT("Animation Notify Track %s does not exist on Animation Sequence %s"), *NotifyTrackName.ToString(), *AnimationSequence->GetName());
+				UE_LOG(LogAnimationBlueprintLibrary, Warning, TEXT("Animation Notify Track %s does not exist on Animation Sequence %s"), *NotifyTrackName.ToString(), *AnimationSequenceBase->GetName());
 			}
 
 			if (!bIsValidTime)
 			{
-				UE_LOG(LogAnimationBlueprintLibrary, Warning, TEXT("%f is side of Animation Sequence %s range"), StartTime, *AnimationSequence->GetName());
+				UE_LOG(LogAnimationBlueprintLibrary, Warning, TEXT("%f is side of Animation Sequence %s range"), StartTime, *AnimationSequenceBase->GetName());
 			}
 		}
 	}
@@ -696,28 +696,28 @@ UAnimNotifyState* UAnimationBlueprintLibrary::AddAnimationNotifyStateEvent(UAnim
 	return NotifyState;
 }
 
-void UAnimationBlueprintLibrary::AddAnimationNotifyEventObject(UAnimSequence* AnimationSequence, float StartTime, UAnimNotify* Notify, FName NotifyTrackName)
+void UAnimationBlueprintLibrary::AddAnimationNotifyEventObject(UAnimSequenceBase* AnimationSequenceBase, float StartTime, UAnimNotify* Notify, FName NotifyTrackName)
 {
-	if (AnimationSequence)
+	if (AnimationSequenceBase)
 	{
 		const bool bValidNotify = Notify != nullptr;
-		const bool bValidOuter = bValidNotify && Notify->GetOuter() == AnimationSequence;
-		const bool bIsValidTrackName = IsValidAnimNotifyTrackName(AnimationSequence, NotifyTrackName);
-		const bool bIsValidTime = IsValidTimeInternal(AnimationSequence, StartTime);
+		const bool bValidOuter = bValidNotify && Notify->GetOuter() == AnimationSequenceBase;
+		const bool bIsValidTrackName = IsValidAnimNotifyTrackName(AnimationSequenceBase, NotifyTrackName);
+		const bool bIsValidTime = IsValidTimeInternal(AnimationSequenceBase, StartTime);
 
 		if (bValidNotify && bValidOuter && bIsValidTrackName && bIsValidTime)
 		{
-			FAnimNotifyEvent& NewEvent = AnimationSequence->Notifies.AddDefaulted_GetRef();
+			FAnimNotifyEvent& NewEvent = AnimationSequenceBase->Notifies.AddDefaulted_GetRef();
 
 			NewEvent.NotifyName = NAME_None;
-			NewEvent.Link(AnimationSequence, StartTime);
-			NewEvent.TriggerTimeOffset = GetTriggerTimeOffsetForType(AnimationSequence->CalculateOffsetForNotify(StartTime));
-			NewEvent.TrackIndex = GetTrackIndexForAnimationNotifyTrackName(AnimationSequence, NotifyTrackName);
+			NewEvent.Link(AnimationSequenceBase, StartTime);
+			NewEvent.TriggerTimeOffset = GetTriggerTimeOffsetForType(AnimationSequenceBase->CalculateOffsetForNotify(StartTime));
+			NewEvent.TrackIndex = GetTrackIndexForAnimationNotifyTrackName(AnimationSequenceBase, NotifyTrackName);
 			NewEvent.NotifyStateClass = nullptr;
 			NewEvent.Notify = Notify;
 
 			// Refresh all cached data
-			AnimationSequence->RefreshCacheData();
+			AnimationSequenceBase->RefreshCacheData();
 		}
 		else
 		{
@@ -729,17 +729,17 @@ void UAnimationBlueprintLibrary::AddAnimationNotifyEventObject(UAnimSequence* An
 			if (!bValidOuter)
 			{
 				const FString NotifyName = bValidNotify ? Notify->GetName() : "Invalid Notify";
-				UE_LOG(LogAnimationBlueprintLibrary, Warning, TEXT("Animation Notify %s Outer is not %s"), *NotifyName, *AnimationSequence->GetName());
+				UE_LOG(LogAnimationBlueprintLibrary, Warning, TEXT("Animation Notify %s Outer is not %s"), *NotifyName, *AnimationSequenceBase->GetName());
 			}
 
 			if (!bIsValidTrackName)
 			{
-				UE_LOG(LogAnimationBlueprintLibrary, Warning, TEXT("Animation Notify Track %s does not exist on Animation Sequence %s"), *NotifyTrackName.ToString(), *AnimationSequence->GetName());
+				UE_LOG(LogAnimationBlueprintLibrary, Warning, TEXT("Animation Notify Track %s does not exist on Animation Sequence %s"), *NotifyTrackName.ToString(), *AnimationSequenceBase->GetName());
 			}
 
 			if (!bIsValidTime)
 			{
-				UE_LOG(LogAnimationBlueprintLibrary, Warning, TEXT("%f is side of Animation Sequence %s range"), StartTime, *AnimationSequence->GetName());
+				UE_LOG(LogAnimationBlueprintLibrary, Warning, TEXT("%f is side of Animation Sequence %s range"), StartTime, *AnimationSequenceBase->GetName());
 			}
 		}
 	}
@@ -749,30 +749,30 @@ void UAnimationBlueprintLibrary::AddAnimationNotifyEventObject(UAnimSequence* An
 	}
 }
 
-void UAnimationBlueprintLibrary::AddAnimationNotifyStateEventObject(UAnimSequence* AnimationSequence, float StartTime, float Duration, UAnimNotifyState* NotifyState, FName NotifyTrackName)
+void UAnimationBlueprintLibrary::AddAnimationNotifyStateEventObject(UAnimSequenceBase* AnimationSequenceBase, float StartTime, float Duration, UAnimNotifyState* NotifyState, FName NotifyTrackName)
 {
-	if (AnimationSequence)
+	if (AnimationSequenceBase)
 	{
 		const bool bValidNotifyState = NotifyState != nullptr;
-		const bool bValidOuter = bValidNotifyState && NotifyState->GetOuter() == AnimationSequence;
-		const bool bIsValidTrackName = IsValidAnimNotifyTrackName(AnimationSequence, NotifyTrackName);
-		const bool bIsValidTime = IsValidTimeInternal(AnimationSequence, StartTime);
+		const bool bValidOuter = bValidNotifyState && NotifyState->GetOuter() == AnimationSequenceBase;
+		const bool bIsValidTrackName = IsValidAnimNotifyTrackName(AnimationSequenceBase, NotifyTrackName);
+		const bool bIsValidTime = IsValidTimeInternal(AnimationSequenceBase, StartTime);
 
 		if (bValidNotifyState && bValidOuter && bIsValidTrackName && bIsValidTime)
 		{
-			FAnimNotifyEvent& NewEvent = AnimationSequence->Notifies.AddDefaulted_GetRef();
+			FAnimNotifyEvent& NewEvent = AnimationSequenceBase->Notifies.AddDefaulted_GetRef();
 
 			NewEvent.NotifyName = NAME_None;
-			NewEvent.Link(AnimationSequence, StartTime);
-			NewEvent.TriggerTimeOffset = GetTriggerTimeOffsetForType(AnimationSequence->CalculateOffsetForNotify(StartTime));
-			NewEvent.TrackIndex = GetTrackIndexForAnimationNotifyTrackName(AnimationSequence, NotifyTrackName);
+			NewEvent.Link(AnimationSequenceBase, StartTime);
+			NewEvent.TriggerTimeOffset = GetTriggerTimeOffsetForType(AnimationSequenceBase->CalculateOffsetForNotify(StartTime));
+			NewEvent.TrackIndex = GetTrackIndexForAnimationNotifyTrackName(AnimationSequenceBase, NotifyTrackName);
 			NewEvent.NotifyStateClass = NotifyState;
 			NewEvent.Notify = nullptr;
 			NewEvent.SetDuration(Duration);
-			NewEvent.EndLink.Link(AnimationSequence, NewEvent.EndLink.GetTime());
+			NewEvent.EndLink.Link(AnimationSequenceBase, NewEvent.EndLink.GetTime());
 
 			// Refresh all cached data
-			AnimationSequence->RefreshCacheData();
+			AnimationSequenceBase->RefreshCacheData();
 		}
 		else
 		{
@@ -784,17 +784,17 @@ void UAnimationBlueprintLibrary::AddAnimationNotifyStateEventObject(UAnimSequenc
 			if (!bValidOuter)
 			{
 				const FString NotifyName = bValidNotifyState ? NotifyState->GetName() : "Invalid Notify";
-				UE_LOG(LogAnimationBlueprintLibrary, Warning, TEXT("Animation Notify %s Outer is not %s"), *NotifyName, *AnimationSequence->GetName());
+				UE_LOG(LogAnimationBlueprintLibrary, Warning, TEXT("Animation Notify %s Outer is not %s"), *NotifyName, *AnimationSequenceBase->GetName());
 			}
 
 			if (!bIsValidTrackName)
 			{
-				UE_LOG(LogAnimationBlueprintLibrary, Warning, TEXT("Animation Notify Track %s does not exist on Animation Sequence %s"), *NotifyTrackName.ToString(), *AnimationSequence->GetName());
+				UE_LOG(LogAnimationBlueprintLibrary, Warning, TEXT("Animation Notify Track %s does not exist on Animation Sequence %s"), *NotifyTrackName.ToString(), *AnimationSequenceBase->GetName());
 			}
 
 			if (!bIsValidTime)
 			{
-				UE_LOG(LogAnimationBlueprintLibrary, Warning, TEXT("%f is side of Animation Sequence %s range"), StartTime, *AnimationSequence->GetName());
+				UE_LOG(LogAnimationBlueprintLibrary, Warning, TEXT("%f is side of Animation Sequence %s range"), StartTime, *AnimationSequenceBase->GetName());
 			}
 		}
 	}
@@ -896,50 +896,50 @@ static void ReplaceAnimNotifies_Helper(UAnimSequenceBase* AnimationSequence, UCl
 	}
 }
 
-void UAnimationBlueprintLibrary::ReplaceAnimNotifyStates(UAnimSequenceBase* AnimationSequence, TSubclassOf<UAnimNotifyState> OldNotifyClass, TSubclassOf<UAnimNotifyState> NewNotifyClass, FOnNotifyStateReplaced OnNotifyStateReplaced)
+void UAnimationBlueprintLibrary::ReplaceAnimNotifyStates(UAnimSequenceBase* AnimationSequenceBase, TSubclassOf<UAnimNotifyState> OldNotifyClass, TSubclassOf<UAnimNotifyState> NewNotifyClass, FOnNotifyStateReplaced OnNotifyStateReplaced)
 {
-	ReplaceAnimNotifies_Helper(AnimationSequence, OldNotifyClass.Get(), NewNotifyClass.Get(), FOnNotifyReplaced(), OnNotifyStateReplaced);
+	ReplaceAnimNotifies_Helper(AnimationSequenceBase, OldNotifyClass.Get(), NewNotifyClass.Get(), FOnNotifyReplaced(), OnNotifyStateReplaced);
 }
 
-void UAnimationBlueprintLibrary::ReplaceAnimNotifies(UAnimSequenceBase* AnimationSequence, TSubclassOf<UAnimNotify> OldNotifyClass, TSubclassOf<UAnimNotify> NewNotifyClass, FOnNotifyReplaced OnNotifyReplaced)
+void UAnimationBlueprintLibrary::ReplaceAnimNotifies(UAnimSequenceBase* AnimationSequenceBase, TSubclassOf<UAnimNotify> OldNotifyClass, TSubclassOf<UAnimNotify> NewNotifyClass, FOnNotifyReplaced OnNotifyReplaced)
 {
-	ReplaceAnimNotifies_Helper(AnimationSequence, OldNotifyClass.Get(), NewNotifyClass.Get(), OnNotifyReplaced, FOnNotifyStateReplaced());
+	ReplaceAnimNotifies_Helper(AnimationSequenceBase, OldNotifyClass.Get(), NewNotifyClass.Get(), OnNotifyReplaced, FOnNotifyStateReplaced());
 }
 
-void UAnimationBlueprintLibrary::CopyAnimNotifiesFromSequence(UAnimSequence* SrcAnimSequence, UAnimSequence* DestAnimSequence)
+void UAnimationBlueprintLibrary::CopyAnimNotifiesFromSequence(UAnimSequenceBase* SourceAnimationSequenceBase, UAnimSequenceBase* DestinationAnimSequenceBase)
 {
-	if (SrcAnimSequence == nullptr)
+	if (SourceAnimationSequenceBase == nullptr)
 	{
 		UE_LOG(LogAnimationBlueprintLibrary, Warning, TEXT("Invalid Source Animation Sequence for CopyAnimNotifiesFromSequence"));
 	}
-	else if (DestAnimSequence == nullptr)
+	else if (DestinationAnimSequenceBase == nullptr)
 	{
 		UE_LOG(LogAnimationBlueprintLibrary, Warning, TEXT("Invalid Destination Animation Sequence for CopyAnimNotifiesFromSequence"));
 	}
-	else if (SrcAnimSequence == DestAnimSequence)
+	else if (SourceAnimationSequenceBase == DestinationAnimSequenceBase)
 	{
 		UE_LOG(LogAnimationBlueprintLibrary, Warning, TEXT("Source and Destination Animation Sequence are the same for CopyAnimNotifiesFromSequence"));
 	}
 	else
 	{
 		const bool bShowDialogs = false;
-		UE::Anim::CopyNotifies(SrcAnimSequence, DestAnimSequence, bShowDialogs);
+		UE::Anim::CopyNotifies(SourceAnimationSequenceBase, DestinationAnimSequenceBase, bShowDialogs);
 	}
 }
 
-int32 UAnimationBlueprintLibrary::RemoveAnimationNotifyEventsByName(UAnimSequence* AnimationSequence, FName NotifyName)
+int32 UAnimationBlueprintLibrary::RemoveAnimationNotifyEventsByName(UAnimSequenceBase* AnimationSequenceBase, FName NotifyName)
 {
 	int32 NumRemovedEvents = 0;
-	if (AnimationSequence)
+	if (AnimationSequenceBase)
 	{
-		NumRemovedEvents = AnimationSequence->Notifies.RemoveAll(
+		NumRemovedEvents = AnimationSequenceBase->Notifies.RemoveAll(
 			[&](const FAnimNotifyEvent& Event)
 		{
 			return Event.NotifyName == NotifyName;
 		});
 
 		// Refresh all cached data
-		AnimationSequence->RefreshCacheData();
+		AnimationSequenceBase->RefreshCacheData();
 	}
 	else
 	{
@@ -949,27 +949,27 @@ int32 UAnimationBlueprintLibrary::RemoveAnimationNotifyEventsByName(UAnimSequenc
 	return NumRemovedEvents;
 }
 
-int32 UAnimationBlueprintLibrary::RemoveAnimationNotifyEventsByTrack(UAnimSequence* AnimationSequence, FName NotifyTrackName)
+int32 UAnimationBlueprintLibrary::RemoveAnimationNotifyEventsByTrack(UAnimSequenceBase* AnimationSequenceBase, FName NotifyTrackName)
 {
 	int32 NumRemovedEvents = 0;
-	if (AnimationSequence)
+	if (AnimationSequenceBase)
 	{
-		const bool bIsValidTrackName = IsValidAnimNotifyTrackName(AnimationSequence, NotifyTrackName);
+		const bool bIsValidTrackName = IsValidAnimNotifyTrackName(AnimationSequenceBase, NotifyTrackName);
 		if (bIsValidTrackName)
 		{
-			const int32 TrackIndex = GetTrackIndexForAnimationNotifyTrackName(AnimationSequence, NotifyTrackName);
-			NumRemovedEvents = AnimationSequence->Notifies.RemoveAll(
+			const int32 TrackIndex = GetTrackIndexForAnimationNotifyTrackName(AnimationSequenceBase, NotifyTrackName);
+			NumRemovedEvents = AnimationSequenceBase->Notifies.RemoveAll(
 				[&](const FAnimNotifyEvent& Event)
 			{
 				return Event.TrackIndex == TrackIndex;
 			});
 
 			// Refresh all cached data
-			AnimationSequence->RefreshCacheData();
+			AnimationSequenceBase->RefreshCacheData();
 		}
 		else
 		{
-			UE_LOG(LogAnimationBlueprintLibrary, Warning, TEXT("Animation Notify Track %s does not exist on Animation Sequence %s"), *NotifyTrackName.ToString(), *AnimationSequence->GetName());
+			UE_LOG(LogAnimationBlueprintLibrary, Warning, TEXT("Animation Notify Track %s does not exist on Animation Sequence %s"), *NotifyTrackName.ToString(), *AnimationSequenceBase->GetName());
 		}
 	}
 	else
@@ -981,12 +981,12 @@ int32 UAnimationBlueprintLibrary::RemoveAnimationNotifyEventsByTrack(UAnimSequen
 	return NumRemovedEvents;
 }
 
-void UAnimationBlueprintLibrary::GetAnimationNotifyTrackNames(const UAnimSequence* AnimationSequence, TArray<FName>& TrackNames)
+void UAnimationBlueprintLibrary::GetAnimationNotifyTrackNames(const UAnimSequenceBase* AnimationSequenceBase, TArray<FName>& TrackNames)
 {
 	TrackNames.Empty();
-	if (AnimationSequence)
+	if (AnimationSequenceBase)
 	{
-		for (const FAnimNotifyTrack& Track : AnimationSequence->AnimNotifyTracks)
+		for (const FAnimNotifyTrack& Track : AnimationSequenceBase->AnimNotifyTracks)
 		{
 			TrackNames.AddUnique(Track.TrackName);
 		}
@@ -997,24 +997,24 @@ void UAnimationBlueprintLibrary::GetAnimationNotifyTrackNames(const UAnimSequenc
 	}
 }
 
-void UAnimationBlueprintLibrary::AddAnimationNotifyTrack(UAnimSequence* AnimationSequence, FName NotifyTrackName, FLinearColor TrackColor /*= FLinearColor::White*/)
+void UAnimationBlueprintLibrary::AddAnimationNotifyTrack(UAnimSequenceBase* AnimationSequenceBase, FName NotifyTrackName, FLinearColor TrackColor /*= FLinearColor::White*/)
 {
-	if (AnimationSequence)
+	if (AnimationSequenceBase)
 	{
-		const bool bExistingTrackName = IsValidAnimNotifyTrackName(AnimationSequence, NotifyTrackName);
+		const bool bExistingTrackName = IsValidAnimNotifyTrackName(AnimationSequenceBase, NotifyTrackName);
 		if (!bExistingTrackName)
 		{
 			FAnimNotifyTrack NewTrack;
 			NewTrack.TrackName = NotifyTrackName;
 			NewTrack.TrackColor = TrackColor;
-			AnimationSequence->AnimNotifyTracks.Add(NewTrack);
+			AnimationSequenceBase->AnimNotifyTracks.Add(NewTrack);
 
 			// Refresh all cached data
-			AnimationSequence->RefreshCacheData();
+			AnimationSequenceBase->RefreshCacheData();
 		}
 		else
 		{
-			UE_LOG(LogAnimationBlueprintLibrary, Warning, TEXT("Animation Notify Track %s already exists on Animation Sequence %s"), *NotifyTrackName.ToString(), *AnimationSequence->GetName());
+			UE_LOG(LogAnimationBlueprintLibrary, Warning, TEXT("Animation Notify Track %s already exists on Animation Sequence %s"), *NotifyTrackName.ToString(), *AnimationSequenceBase->GetName());
 		}
 	}
 	else
@@ -1024,38 +1024,42 @@ void UAnimationBlueprintLibrary::AddAnimationNotifyTrack(UAnimSequence* Animatio
 	
 }
 
-void UAnimationBlueprintLibrary::RemoveAnimationNotifyTrack(UAnimSequence* AnimationSequence, FName NotifyTrackName)
+void UAnimationBlueprintLibrary::RemoveAnimationNotifyTrack(UAnimSequenceBase* AnimationSequenceBase, FName NotifyTrackName)
 {
-	if (AnimationSequence)
+	if (AnimationSequenceBase)
 	{
-		const int32 TrackIndexToDelete = GetTrackIndexForAnimationNotifyTrackName(AnimationSequence, NotifyTrackName);
+		const int32 TrackIndexToDelete = GetTrackIndexForAnimationNotifyTrackName(AnimationSequenceBase, NotifyTrackName);
 		if (TrackIndexToDelete != INDEX_NONE)
 		{	
 			// Remove all notifies and sync markers on the to-delete-track
-			AnimationSequence->Notifies.RemoveAll([&](const FAnimNotifyEvent& Notify) { return Notify.TrackIndex == TrackIndexToDelete; });
-			AnimationSequence->AuthoredSyncMarkers.RemoveAll([&](const FAnimSyncMarker& Marker) { return Marker.TrackIndex == TrackIndexToDelete; });
-
+			AnimationSequenceBase->Notifies.RemoveAll([&](const FAnimNotifyEvent& Notify) { return Notify.TrackIndex == TrackIndexToDelete; });
+		
 			// Before track removal, make sure everything behind is fixed
-			for (FAnimNotifyEvent& Notify : AnimationSequence->Notifies)
+			for (FAnimNotifyEvent& Notify : AnimationSequenceBase->Notifies)
 			{
 				if (Notify.TrackIndex > TrackIndexToDelete)
 				{
 					Notify.TrackIndex = Notify.TrackIndex - 1;
 				}				
 			}
-			for (FAnimSyncMarker& SyncMarker : AnimationSequence->AuthoredSyncMarkers)
+
+			if (UAnimSequence* AnimationSequence = Cast<UAnimSequence>(AnimationSequenceBase))
 			{
-				if (SyncMarker.TrackIndex > TrackIndexToDelete)
+				AnimationSequence->AuthoredSyncMarkers.RemoveAll([&](const FAnimSyncMarker& Marker) { return Marker.TrackIndex == TrackIndexToDelete; });
+				for (FAnimSyncMarker& SyncMarker : AnimationSequence->AuthoredSyncMarkers)
 				{
-					SyncMarker.TrackIndex = SyncMarker.TrackIndex - 1;
+					if (SyncMarker.TrackIndex > TrackIndexToDelete)
+					{
+						SyncMarker.TrackIndex = SyncMarker.TrackIndex - 1;
+					}
 				}
-			}
+			}			
 			
 			// Delete the track itself
-			AnimationSequence->AnimNotifyTracks.RemoveAt(TrackIndexToDelete);
+			AnimationSequenceBase->AnimNotifyTracks.RemoveAt(TrackIndexToDelete);
 
 			// Refresh all cached data
-			AnimationSequence->RefreshCacheData();
+			AnimationSequenceBase->RefreshCacheData();
 		}		
 	}
 	else
@@ -1064,22 +1068,25 @@ void UAnimationBlueprintLibrary::RemoveAnimationNotifyTrack(UAnimSequence* Anima
 	}
 }
 
-void UAnimationBlueprintLibrary::RemoveAllAnimationNotifyTracks(UAnimSequence* AnimationSequence)
+void UAnimationBlueprintLibrary::RemoveAllAnimationNotifyTracks(UAnimSequenceBase* AnimationSequenceBase)
 {
-	if (AnimationSequence)
+	if (AnimationSequenceBase)
 	{
-		AnimationSequence->Notifies.Empty();
-		AnimationSequence->AuthoredSyncMarkers.Empty();
+		AnimationSequenceBase->Notifies.Empty();
+		if (UAnimSequence* AnimationSequence = Cast<UAnimSequence>(AnimationSequenceBase))
+		{
+			AnimationSequence->AuthoredSyncMarkers.Empty();
+		}
 
 		// Remove all but one notify tracks
-		AnimationSequence->AnimNotifyTracks.SetNum(1);
+		AnimationSequenceBase->AnimNotifyTracks.SetNum(1);
 
 		// Also remove all stale notifies and sync markers from only track 
-		AnimationSequence->AnimNotifyTracks[0].Notifies.Empty();
-		AnimationSequence->AnimNotifyTracks[0].SyncMarkers.Empty();
+		AnimationSequenceBase->AnimNotifyTracks[0].Notifies.Empty();
+		AnimationSequenceBase->AnimNotifyTracks[0].SyncMarkers.Empty();
 
 		// Refresh all cached data
-		AnimationSequence->RefreshCacheData();
+		AnimationSequenceBase->RefreshCacheData();
 	}
 	else
 	{
@@ -1087,12 +1094,12 @@ void UAnimationBlueprintLibrary::RemoveAllAnimationNotifyTracks(UAnimSequence* A
 	}	
 }
 
-bool UAnimationBlueprintLibrary::IsValidAnimNotifyTrackName(const UAnimSequence* AnimationSequence, FName NotifyTrackName)
+bool UAnimationBlueprintLibrary::IsValidAnimNotifyTrackName(const UAnimSequenceBase* AnimationSequenceBase, FName NotifyTrackName)
 {
 	bool bIsValid = false;
-	if (AnimationSequence)
+	if (AnimationSequenceBase)
 	{
-		bIsValid = GetTrackIndexForAnimationNotifyTrackName(AnimationSequence, NotifyTrackName) != INDEX_NONE;
+		bIsValid = GetTrackIndexForAnimationNotifyTrackName(AnimationSequenceBase, NotifyTrackName) != INDEX_NONE;
 	}
 	else
 	{
@@ -1102,20 +1109,20 @@ bool UAnimationBlueprintLibrary::IsValidAnimNotifyTrackName(const UAnimSequence*
 	return bIsValid;	
 }
 
-int32 UAnimationBlueprintLibrary::GetTrackIndexForAnimationNotifyTrackName(const UAnimSequence* AnimationSequence, FName NotifyTrackName)
+int32 UAnimationBlueprintLibrary::GetTrackIndexForAnimationNotifyTrackName(const UAnimSequenceBase* AnimationSequenceBase, FName NotifyTrackName)
 {
-	return AnimationSequence->AnimNotifyTracks.IndexOfByPredicate(
+	return AnimationSequenceBase->AnimNotifyTracks.IndexOfByPredicate(
 		[&](const FAnimNotifyTrack& Track)
 	{
 		return Track.TrackName == NotifyTrackName;
 	});
 }
 
-const FAnimNotifyTrack& UAnimationBlueprintLibrary::GetNotifyTrackByName(const UAnimSequence* AnimationSequence, FName NotifyTrackName)
+const FAnimNotifyTrack& UAnimationBlueprintLibrary::GetNotifyTrackByName(const UAnimSequenceBase* AnimationSequenceBase, FName NotifyTrackName)
 {
-	const int32 TrackIndex = GetTrackIndexForAnimationNotifyTrackName(AnimationSequence, NotifyTrackName);
-	checkf(TrackIndex != INDEX_NONE, TEXT("Notify Track %s does not exist on %s"), *NotifyTrackName.ToString(), *AnimationSequence->GetName());
-	return AnimationSequence->AnimNotifyTracks[TrackIndex];
+	const int32 TrackIndex = GetTrackIndexForAnimationNotifyTrackName(AnimationSequenceBase, NotifyTrackName);
+	checkf(TrackIndex != INDEX_NONE, TEXT("Notify Track %s does not exist on %s"), *NotifyTrackName.ToString(), *AnimationSequenceBase->GetName());
+	return AnimationSequenceBase->AnimNotifyTracks[TrackIndex];
 }
 
 float UAnimationBlueprintLibrary::GetAnimNotifyEventTriggerTime(const FAnimNotifyEvent& NotifyEvent)
@@ -1150,16 +1157,16 @@ void UAnimationBlueprintLibrary::GetAnimationSyncMarkersForTrack(const UAnimSequ
 	}
 }
 
-void UAnimationBlueprintLibrary::GetAnimationNotifyEventsForTrack(const UAnimSequence* AnimationSequence, FName NotifyTrackName, TArray<FAnimNotifyEvent>& Events)
+void UAnimationBlueprintLibrary::GetAnimationNotifyEventsForTrack(const UAnimSequenceBase* AnimationSequenceBase, FName NotifyTrackName, TArray<FAnimNotifyEvent>& Events)
 {
 	Events.Empty();
-	if (AnimationSequence)
+	if (AnimationSequenceBase)
 	{
-		const bool bIsValidTrackName = IsValidAnimNotifyTrackName(AnimationSequence, NotifyTrackName);
+		const bool bIsValidTrackName = IsValidAnimNotifyTrackName(AnimationSequenceBase, NotifyTrackName);
 
 		if (bIsValidTrackName)
 		{
-			const FAnimNotifyTrack& Track = GetNotifyTrackByName(AnimationSequence, NotifyTrackName);
+			const FAnimNotifyTrack& Track = GetNotifyTrackByName(AnimationSequenceBase, NotifyTrackName);
 			Events.Empty(Track.Notifies.Num());
 			for (FAnimNotifyEvent* Event : Track.Notifies)
 			{
@@ -1168,7 +1175,7 @@ void UAnimationBlueprintLibrary::GetAnimationNotifyEventsForTrack(const UAnimSeq
 		}
 		else
 		{
-			UE_LOG(LogAnimationBlueprintLibrary, Warning, TEXT("Animation Notify Track %s does not exist on Animation Sequence %s"), *NotifyTrackName.ToString(), *AnimationSequence->GetName());
+			UE_LOG(LogAnimationBlueprintLibrary, Warning, TEXT("Animation Notify Track %s does not exist on Animation Sequence %s"), *NotifyTrackName.ToString(), *AnimationSequenceBase->GetName());
 		}
 	}
 	else
@@ -1660,14 +1667,14 @@ FName UAnimationBlueprintLibrary::RetrieveContainerNameForCurve(const UAnimSeque
 	return NAME_None;
 }
 
-void UAnimationBlueprintLibrary::AddMetaData(UAnimSequence* AnimationSequence, TSubclassOf<UAnimMetaData> MetaDataClass, UAnimMetaData*& MetaDataInstance)
+void UAnimationBlueprintLibrary::AddMetaData(UAnimationAsset* AnimationAsset, TSubclassOf<UAnimMetaData> MetaDataClass, UAnimMetaData*& MetaDataInstance)
 {
-	if (AnimationSequence)
+	if (AnimationAsset)
 	{
-		MetaDataInstance = NewObject<UAnimMetaData>(AnimationSequence, MetaDataClass, NAME_None, RF_Transactional);
+		MetaDataInstance = NewObject<UAnimMetaData>(AnimationAsset, MetaDataClass, NAME_None, RF_Transactional);
 		if (MetaDataInstance)
 		{
-			AnimationSequence->AddMetaData(MetaDataInstance);
+			AnimationAsset->AddMetaData(MetaDataInstance);
 		}
 		else
 		{
@@ -1681,22 +1688,22 @@ void UAnimationBlueprintLibrary::AddMetaData(UAnimSequence* AnimationSequence, T
 	}
 }
 
-void UAnimationBlueprintLibrary::AddMetaDataObject(UAnimSequence* AnimationSequence, UAnimMetaData* MetaDataObject)
+void UAnimationBlueprintLibrary::AddMetaDataObject(UAnimationAsset* AnimationAsset, UAnimMetaData* MetaDataObject)
 {
-	if (AnimationSequence && MetaDataObject)
+	if (AnimationAsset && MetaDataObject)
 	{
-		if (MetaDataObject->GetOuter() == AnimationSequence)
+		if (MetaDataObject->GetOuter() == AnimationAsset)
 		{
-			AnimationSequence->AddMetaData(MetaDataObject);
+			AnimationAsset->AddMetaData(MetaDataObject);
 		}
 		else
 		{
-			UE_LOG(LogAnimationBlueprintLibrary, Warning, TEXT("Outer for MetaData Instance %s is not Animation Sequence %s"), *MetaDataObject->GetName(), *AnimationSequence->GetName());
+			UE_LOG(LogAnimationBlueprintLibrary, Warning, TEXT("Outer for MetaData Instance %s is not Animation Sequence %s"), *MetaDataObject->GetName(), *AnimationAsset->GetName());
 		}		
 	}
 	else 
 	{
-		if (!AnimationSequence)
+		if (!AnimationAsset)
 		{
 			UE_LOG(LogAnimationBlueprintLibrary, Warning, TEXT("Invalid Animation Sequence for AddMetaDataObject"));
 		}
@@ -1708,11 +1715,11 @@ void UAnimationBlueprintLibrary::AddMetaDataObject(UAnimSequence* AnimationSeque
 	}
 }
 
-void UAnimationBlueprintLibrary::RemoveAllMetaData(UAnimSequence* AnimationSequence)
+void UAnimationBlueprintLibrary::RemoveAllMetaData(UAnimationAsset* AnimationAsset)
 {
-	if (AnimationSequence)
+	if (AnimationAsset)
 	{
-		AnimationSequence->EmptyMetaData();
+		AnimationAsset->EmptyMetaData();
 	}
 	else
 	{
@@ -1720,11 +1727,11 @@ void UAnimationBlueprintLibrary::RemoveAllMetaData(UAnimSequence* AnimationSeque
 	}
 }
 
-void UAnimationBlueprintLibrary::RemoveMetaData(UAnimSequence* AnimationSequence, UAnimMetaData* MetaDataObject)
+void UAnimationBlueprintLibrary::RemoveMetaData(UAnimationAsset* AnimationAsset, UAnimMetaData* MetaDataObject)
 {
-	if (AnimationSequence)
+	if (AnimationAsset)
 	{
-		AnimationSequence->RemoveMetaData(MetaDataObject);
+		AnimationAsset->RemoveMetaData(MetaDataObject);
 	}
 	else
 	{
@@ -1732,13 +1739,13 @@ void UAnimationBlueprintLibrary::RemoveMetaData(UAnimSequence* AnimationSequence
 	}
 }
 
-void UAnimationBlueprintLibrary::RemoveMetaDataOfClass(UAnimSequence* AnimationSequence, TSubclassOf<UAnimMetaData> MetaDataClass)
+void UAnimationBlueprintLibrary::RemoveMetaDataOfClass(UAnimationAsset* AnimationAsset, TSubclassOf<UAnimMetaData> MetaDataClass)
 {
-	if (AnimationSequence)
+	if (AnimationAsset)
 	{
 		TArray<UAnimMetaData*> MetaDataOfClass;
-		GetMetaDataOfClass(AnimationSequence, MetaDataClass, MetaDataOfClass);
-		AnimationSequence->RemoveMetaData(MetaDataOfClass);
+		GetMetaDataOfClass(AnimationAsset, MetaDataClass, MetaDataOfClass);
+		AnimationAsset->RemoveMetaData(MetaDataOfClass);
 	}
 	else
 	{
@@ -1746,12 +1753,12 @@ void UAnimationBlueprintLibrary::RemoveMetaDataOfClass(UAnimSequence* AnimationS
 	}	
 }
 
-void UAnimationBlueprintLibrary::GetMetaData(const UAnimSequence* AnimationSequence, TArray<UAnimMetaData*>& MetaData)
+void UAnimationBlueprintLibrary::GetMetaData(const UAnimationAsset* AnimationAsset, TArray<UAnimMetaData*>& MetaData)
 {
 	MetaData.Empty();
-	if (AnimationSequence)
+	if (AnimationAsset)
 	{
-		MetaData = AnimationSequence->GetMetaData();
+		MetaData = AnimationAsset->GetMetaData();
 	}
 	else
 	{
@@ -1759,12 +1766,12 @@ void UAnimationBlueprintLibrary::GetMetaData(const UAnimSequence* AnimationSeque
 	}
 }
 
-void UAnimationBlueprintLibrary::GetMetaDataOfClass(const UAnimSequence* AnimationSequence, TSubclassOf<UAnimMetaData> MetaDataClass, TArray<UAnimMetaData*>& MetaDataOfClass)
+void UAnimationBlueprintLibrary::GetMetaDataOfClass(const UAnimationAsset* AnimationAsset, TSubclassOf<UAnimMetaData> MetaDataClass, TArray<UAnimMetaData*>& MetaDataOfClass)
 {
 	MetaDataOfClass.Empty();
-	if (AnimationSequence)
+	if (AnimationAsset)
 	{
-		const TArray<UAnimMetaData*>& MetaData = AnimationSequence->GetMetaData();
+		const TArray<UAnimMetaData*>& MetaData = AnimationAsset->GetMetaData();
 		for (UAnimMetaData* MetaDataInstance : MetaData)
 		{
 			if (MetaDataInstance->GetClass() == *MetaDataClass)
@@ -1779,13 +1786,13 @@ void UAnimationBlueprintLibrary::GetMetaDataOfClass(const UAnimSequence* Animati
 	}
 }
 
-bool UAnimationBlueprintLibrary::ContainsMetaDataOfClass(const UAnimSequence* AnimationSequence, TSubclassOf<UAnimMetaData> MetaDataClass)
+bool UAnimationBlueprintLibrary::ContainsMetaDataOfClass(const UAnimationAsset* AnimationAsset, TSubclassOf<UAnimMetaData> MetaDataClass)
 {
 	bool bContainsMetaData = false;
-	if (AnimationSequence)
+	if (AnimationAsset)
 	{
 		TArray<UAnimMetaData*> MetaData;
-		GetMetaData(AnimationSequence, MetaData);
+		GetMetaData(AnimationAsset, MetaData);
 		bContainsMetaData = MetaData.FindByPredicate(
 			[&](UAnimMetaData* MetaDataObject)
 		{
@@ -1800,15 +1807,15 @@ bool UAnimationBlueprintLibrary::ContainsMetaDataOfClass(const UAnimSequence* An
 	return bContainsMetaData;	
 }
 
-void UAnimationBlueprintLibrary::GetBonePoseForTime(const UAnimSequence* AnimationSequence, FName BoneName, float Time, bool bExtractRootMotion, FTransform& Pose)
+void UAnimationBlueprintLibrary::GetBonePoseForTime(const UAnimSequenceBase* AnimationSequenceBase, FName BoneName, float Time, bool bExtractRootMotion, FTransform& Pose)
 {
 	Pose.SetIdentity();
-	if (AnimationSequence)
+	if (AnimationSequenceBase)
 	{
 		TArray<FName> BoneNameArray;
 		TArray<FTransform> PoseArray;
 		BoneNameArray.Add(BoneName);
-		GetBonePosesForTime(AnimationSequence, BoneNameArray, Time, bExtractRootMotion, PoseArray);
+		GetBonePosesForTime(AnimationSequenceBase, BoneNameArray, Time, bExtractRootMotion, PoseArray);
 		Pose = PoseArray[0];
 	}
 	else
@@ -1817,12 +1824,12 @@ void UAnimationBlueprintLibrary::GetBonePoseForTime(const UAnimSequence* Animati
 	}
 }
 
-void UAnimationBlueprintLibrary::GetBonePoseForFrame(const UAnimSequence* AnimationSequence, FName BoneName, int32 Frame, bool bExtractRootMotion, FTransform& Pose)
+void UAnimationBlueprintLibrary::GetBonePoseForFrame(const UAnimSequenceBase* AnimationSequenceBase, FName BoneName, int32 Frame, bool bExtractRootMotion, FTransform& Pose)
 {
 	Pose.SetIdentity();
-	if (AnimationSequence)
+	if (AnimationSequenceBase)
 	{
-		GetBonePoseForTime(AnimationSequence, BoneName, GetTimeAtFrameInternal(AnimationSequence, Frame), bExtractRootMotion, Pose);
+		GetBonePoseForTime(AnimationSequenceBase, BoneName, GetTimeAtFrameInternal(AnimationSequenceBase, Frame), bExtractRootMotion, Pose);
 	}
 	else
 	{
@@ -1830,19 +1837,19 @@ void UAnimationBlueprintLibrary::GetBonePoseForFrame(const UAnimSequence* Animat
 	}
 }
 
-void UAnimationBlueprintLibrary::GetBonePosesForTime(const UAnimSequence* AnimationSequence, TArray<FName> BoneNames, float Time, bool bExtractRootMotion, TArray<FTransform>& Poses, const USkeletalMesh* PreviewMesh /*= nullptr*/)
+void UAnimationBlueprintLibrary::GetBonePosesForTime(const UAnimSequenceBase* AnimationSequenceBase, TArray<FName> BoneNames, float Time, bool bExtractRootMotion, TArray<FTransform>& Poses, const USkeletalMesh* PreviewMesh /*= nullptr*/)
 {
 	Poses.Empty(BoneNames.Num());
-	if (AnimationSequence && AnimationSequence->GetSkeleton())
+	if (AnimationSequenceBase && AnimationSequenceBase->GetSkeleton())
 	{
 		Poses.AddDefaulted(BoneNames.Num());
 
 		// Need this for FCompactPose
 		FMemMark Mark(FMemStack::Get());
 
-		const FReferenceSkeleton& RefSkeleton = (PreviewMesh)? PreviewMesh->GetRefSkeleton() : AnimationSequence->GetSkeleton()->GetReferenceSkeleton();
+		const FReferenceSkeleton& RefSkeleton = (PreviewMesh)? PreviewMesh->GetRefSkeleton() : AnimationSequenceBase->GetSkeleton()->GetReferenceSkeleton();
 
-		if (IsValidTimeInternal(AnimationSequence, Time))
+		if (IsValidTimeInternal(AnimationSequenceBase, Time))
 		{
 			if (BoneNames.Num())
 			{
@@ -1851,9 +1858,18 @@ void UAnimationBlueprintLibrary::GetBonePosesForTime(const UAnimSequence* Animat
 					const FName& BoneName = BoneNames[BoneNameIndex];
 					
 					FTransform& Transform = Poses[BoneNameIndex];
-					if (IsValidRawAnimationTrackName(AnimationSequence, BoneName))
+					if (IsValidRawAnimationTrackName(AnimationSequenceBase, BoneName))
 					{
-						UE::Anim::GetBoneTransformFromModel(AnimationSequence->GetDataModel(), Transform, AnimationSequence->GetDataModel()->GetBoneTrackIndexByName(BoneName), Time, AnimationSequence->Interpolation);
+						const EAnimInterpolationType InterpolationType = [AnimationSequenceBase]() -> EAnimInterpolationType
+						{
+							if (const UAnimSequence* AnimationSequence = Cast<const UAnimSequence>(AnimationSequenceBase))
+							{
+								return AnimationSequence->Interpolation;
+							}
+
+							return EAnimInterpolationType::Linear;
+						}();
+						UE::Anim::GetBoneTransformFromModel(AnimationSequenceBase->GetDataModel(), Transform, AnimationSequenceBase->GetDataModel()->GetBoneTrackIndexByName(BoneName), Time, InterpolationType);
 					}
 					else
 					{
@@ -1866,7 +1882,7 @@ void UAnimationBlueprintLibrary::GetBonePosesForTime(const UAnimSequence* Animat
 						}
 						else
 						{
-							UE_LOG(LogAnimationBlueprintLibrary, Warning, TEXT("Invalid bone name %s for Animation Sequence %s supplied for GetBonePosesForTime"), *BoneName.ToString(), *AnimationSequence->GetName());
+							UE_LOG(LogAnimationBlueprintLibrary, Warning, TEXT("Invalid bone name %s for Animation Sequence %s supplied for GetBonePosesForTime"), *BoneName.ToString(), *AnimationSequenceBase->GetName());
 							Transform = FTransform::Identity;
 						}
 					}
@@ -1874,12 +1890,12 @@ void UAnimationBlueprintLibrary::GetBonePosesForTime(const UAnimSequence* Animat
 			}
 			else
 			{
-				UE_LOG(LogAnimationBlueprintLibrary, Error, TEXT("Invalid or no bone names specified to retrieve poses given Animation Sequence %s in GetBonePosesForTime"), *AnimationSequence->GetName());
+				UE_LOG(LogAnimationBlueprintLibrary, Error, TEXT("Invalid or no bone names specified to retrieve poses given Animation Sequence %s in GetBonePosesForTime"), *AnimationSequenceBase->GetName());
 			}			
 		}
 		else
 		{
-			UE_LOG(LogAnimationBlueprintLibrary, Warning, TEXT("Invalid time value %f for Animation Sequence %s supplied for GetBonePosesForTime"), Time, *AnimationSequence->GetName());
+			UE_LOG(LogAnimationBlueprintLibrary, Warning, TEXT("Invalid time value %f for Animation Sequence %s supplied for GetBonePosesForTime"), Time, *AnimationSequenceBase->GetName());
 		}
 	}
 	else
@@ -1888,12 +1904,12 @@ void UAnimationBlueprintLibrary::GetBonePosesForTime(const UAnimSequence* Animat
 	}
 }
 
-void UAnimationBlueprintLibrary::GetBonePosesForFrame(const UAnimSequence* AnimationSequence, TArray<FName> BoneNames, int32 Frame, bool bExtractRootMotion, TArray<FTransform>& Poses, const USkeletalMesh* PreviewMesh /*= nullptr*/)
+void UAnimationBlueprintLibrary::GetBonePosesForFrame(const UAnimSequenceBase* AnimationSequenceBase, TArray<FName> BoneNames, int32 Frame, bool bExtractRootMotion, TArray<FTransform>& Poses, const USkeletalMesh* PreviewMesh /*= nullptr*/)
 {
 	Poses.Empty(BoneNames.Num());
-	if (AnimationSequence)
+	if (AnimationSequenceBase)
 	{
-		GetBonePosesForTime(AnimationSequence, BoneNames, GetTimeAtFrameInternal(AnimationSequence, Frame), bExtractRootMotion, Poses, PreviewMesh);
+		GetBonePosesForTime(AnimationSequenceBase, BoneNames, GetTimeAtFrameInternal(AnimationSequenceBase, Frame), bExtractRootMotion, Poses, PreviewMesh);
 	}
 	else
 	{
@@ -2039,12 +2055,12 @@ bool UAnimationBlueprintLibrary::DoesVirtualBoneNameExistInternal(USkeleton* Ske
 	return Skeleton->VirtualBones.IndexOfByPredicate([&](const FVirtualBone& VirtualBone) { return VirtualBone.VirtualBoneName == BoneName; }) != INDEX_NONE;
 }
 
-void UAnimationBlueprintLibrary::GetSequenceLength(const UAnimSequence* AnimationSequence, float& Length)
+void UAnimationBlueprintLibrary::GetSequenceLength(const UAnimSequenceBase* AnimationSequenceBase, float& Length)
 {
 	Length = 0.0f;
-	if (AnimationSequence)
+	if (AnimationSequenceBase)
 	{
-		Length = AnimationSequence->GetPlayLength();
+		Length = AnimationSequenceBase->GetPlayLength();
 	}
 	else
 	{
@@ -2052,12 +2068,12 @@ void UAnimationBlueprintLibrary::GetSequenceLength(const UAnimSequence* Animatio
 	}
 }
 
-void UAnimationBlueprintLibrary::GetRateScale(const UAnimSequence* AnimationSequence, float& RateScale)
+void UAnimationBlueprintLibrary::GetRateScale(const UAnimSequenceBase* AnimationSequenceBase, float& RateScale)
 {
 	RateScale = 0.0f;
-	if (AnimationSequence)
+	if (AnimationSequenceBase)
 	{
-		RateScale = AnimationSequence->RateScale;
+		RateScale = AnimationSequenceBase->RateScale;
 	}
 	else
 	{
@@ -2065,11 +2081,11 @@ void UAnimationBlueprintLibrary::GetRateScale(const UAnimSequence* AnimationSequ
 	}
 }
 
-void UAnimationBlueprintLibrary::SetRateScale(UAnimSequence* AnimationSequence, float RateScale)
+void UAnimationBlueprintLibrary::SetRateScale(UAnimSequenceBase* AnimationSequenceBase, float RateScale)
 {	
-	if (AnimationSequence)
+	if (AnimationSequenceBase)
 	{
-		AnimationSequence->RateScale = RateScale;
+		AnimationSequenceBase->RateScale = RateScale;
 	}
 	else
 	{
@@ -2077,12 +2093,12 @@ void UAnimationBlueprintLibrary::SetRateScale(UAnimSequence* AnimationSequence,
 	}
 }
 
-void UAnimationBlueprintLibrary::GetFrameAtTime(const UAnimSequence* AnimationSequence, const float Time, int32& Frame)
+void UAnimationBlueprintLibrary::GetFrameAtTime(const UAnimSequenceBase* AnimationSequenceBase, const float Time, int32& Frame)
 {
 	Frame = 0;
-	if (AnimationSequence)
+	if (AnimationSequenceBase)
 	{
-		Frame = AnimationSequence->GetFrameAtTime(Time);
+		Frame = AnimationSequenceBase->GetFrameAtTime(Time);
 	}
 	else
 	{		
@@ -2090,12 +2106,12 @@ void UAnimationBlueprintLibrary::GetFrameAtTime(const UAnimSequence* AnimationSe
 	}
 }
 
-void UAnimationBlueprintLibrary::GetTimeAtFrame(const UAnimSequence* AnimationSequence, const int32 Frame, float& Time)
+void UAnimationBlueprintLibrary::GetTimeAtFrame(const UAnimSequenceBase* AnimationSequenceBase, const int32 Frame, float& Time)
 {
 	Time = 0.0f;
-	if (AnimationSequence)
+	if (AnimationSequenceBase)
 	{
-		Time = GetTimeAtFrameInternal(AnimationSequence, Frame);
+		Time = GetTimeAtFrameInternal(AnimationSequenceBase, Frame);
 	}
 	else
 	{		
@@ -2103,17 +2119,17 @@ void UAnimationBlueprintLibrary::GetTimeAtFrame(const UAnimSequence* AnimationSe
 	}
 }
 
-float UAnimationBlueprintLibrary::GetTimeAtFrameInternal(const UAnimSequence* AnimationSequence, const int32 Frame)
+float UAnimationBlueprintLibrary::GetTimeAtFrameInternal(const UAnimSequenceBase* AnimationSequenceBase, const int32 Frame)
 {
-	return AnimationSequence->GetTimeAtFrame(Frame);
+	return AnimationSequenceBase->GetTimeAtFrame(Frame);
 }
 
-void UAnimationBlueprintLibrary::IsValidTime(const UAnimSequence* AnimationSequence, const float Time, bool& IsValid)
+void UAnimationBlueprintLibrary::IsValidTime(const UAnimSequenceBase* AnimationSequenceBase, const float Time, bool& IsValid)
 {
 	IsValid = false;
-	if (AnimationSequence)
+	if (AnimationSequenceBase)
 	{
-		IsValid = IsValidTimeInternal(AnimationSequence, Time);
+		IsValid = IsValidTimeInternal(AnimationSequenceBase, Time);
 	}
 	else
 	{		
@@ -2121,26 +2137,26 @@ void UAnimationBlueprintLibrary::IsValidTime(const UAnimSequence* AnimationSeque
 	}
 }
 
-bool UAnimationBlueprintLibrary::IsValidTimeInternal(const UAnimSequence* AnimationSequence, const float Time)
+bool UAnimationBlueprintLibrary::IsValidTimeInternal(const UAnimSequenceBase* AnimationSequenceBase, const float Time)
 {
-	return FMath::IsWithinInclusive(Time, 0.0f, AnimationSequence->GetPlayLength());
+	return FMath::IsWithinInclusive(Time, 0.0f, AnimationSequenceBase->GetPlayLength());
 }
 
-void UAnimationBlueprintLibrary::FindBonePathToRoot(const UAnimSequence* AnimationSequence, FName BoneName, TArray<FName>& BonePath)
+void UAnimationBlueprintLibrary::FindBonePathToRoot(const UAnimSequenceBase* AnimationSequenceBase, FName BoneName, TArray<FName>& BonePath)
 {
 	BonePath.Empty();
-	if (AnimationSequence)
+	if (AnimationSequenceBase)
 	{
 		BonePath.Add(BoneName);
-		int32 BoneIndex = AnimationSequence->GetSkeleton()->GetReferenceSkeleton().FindRawBoneIndex(BoneName);		
+		int32 BoneIndex = AnimationSequenceBase->GetSkeleton()->GetReferenceSkeleton().FindRawBoneIndex(BoneName);		
 		if (BoneIndex != INDEX_NONE)
 		{
 			while (BoneIndex != INDEX_NONE)
 			{
-				const int32 ParentBoneIndex = AnimationSequence->GetSkeleton()->GetReferenceSkeleton().GetRawParentIndex(BoneIndex);
+				const int32 ParentBoneIndex = AnimationSequenceBase->GetSkeleton()->GetReferenceSkeleton().GetRawParentIndex(BoneIndex);
 				if (ParentBoneIndex != INDEX_NONE)
 				{
-					BonePath.Add(AnimationSequence->GetSkeleton()->GetReferenceSkeleton().GetBoneName(ParentBoneIndex));
+					BonePath.Add(AnimationSequenceBase->GetSkeleton()->GetReferenceSkeleton().GetBoneName(ParentBoneIndex));
 				}
 
 				BoneIndex = ParentBoneIndex;
@@ -2148,7 +2164,7 @@ void UAnimationBlueprintLibrary::FindBonePathToRoot(const UAnimSequence* Animati
 		}
 		else
 		{
-			UE_LOG(LogAnimationBlueprintLibrary, Warning, TEXT("Bone name %s not found in Skeleton %s"), *BoneName.ToString(), *AnimationSequence->GetSkeleton()->GetName());
+			UE_LOG(LogAnimationBlueprintLibrary, Warning, TEXT("Bone name %s not found in Skeleton %s"), *BoneName.ToString(), *AnimationSequenceBase->GetSkeleton()->GetName());
 		}
 	}
 	else
diff --git a/Engine/Source/Editor/AnimationModifiers/Public/AnimationBlueprintLibrary.h b/Engine/Source/Editor/AnimationModifiers/Public/AnimationBlueprintLibrary.h
--- a/Engine/Source/Editor/AnimationModifiers/Public/AnimationBlueprintLibrary.h
+++ b/Engine/Source/Editor/AnimationModifiers/Public/AnimationBlueprintLibrary.h
@@ -38,15 +38,15 @@ class ANIMATIONMODIFIERS_API UAnimationBlueprintLibrary : public UBlueprintFunct
 public:
 	/** Retrieves the number of animation frames for the given Animation Sequence */
 	UFUNCTION(BlueprintPure, meta=(AutoCreateRefTerm = "AnimationSequence"), Category = "AnimationBlueprintLibrary|Animation")
-	static void GetNumFrames(const UAnimSequence* AnimationSequence, int32& NumFrames);
+	static void GetNumFrames(const UAnimSequenceBase* AnimationSequenceBase, int32& NumFrames);
 
 	/** Retrieves the number of animation keys for the given Animation Sequence */
 	UFUNCTION(BlueprintPure, meta = (AutoCreateRefTerm = "AnimationSequence"), Category = "AnimationBlueprintLibrary|Animation")
-	static void GetNumKeys(const UAnimSequence* AnimationSequence, int32& NumKeys);
+	static void GetNumKeys(const UAnimSequenceBase* AnimationSequenceBase, int32& NumKeys);
 	
 	/** Retrieves the Names of the individual ATracks for the given Animation Sequence */
 	UFUNCTION(BlueprintPure, Category = "AnimationBlueprintLibrary|Animation")
-	static void GetAnimationTrackNames(const UAnimSequence* AnimationSequence, TArray<FName>& TrackNames);
+	static void GetAnimationTrackNames(const UAnimSequenceBase* AnimationSequenceBase, TArray<FName>& TrackNames);
 
 	/** Retrieves the Names of the individual float curves for the given Animation Sequence */
 	UFUNCTION(BlueprintPure, Category = "AnimationBlueprintLibrary|Animation")
@@ -54,25 +54,25 @@ public:
 
 	/** Retrieves the Raw Translation Animation Data for the given Animation Track Name and Animation Sequence */
 	UFUNCTION(BlueprintPure, Category = "AnimationBlueprintLibrary|RawTrackData")
-	static void GetRawTrackPositionData(const UAnimSequence* AnimationSequence, const FName TrackName, TArray<FVector>& PositionData);
+	static void GetRawTrackPositionData(const UAnimSequenceBase* AnimationSequenceBase, const FName TrackName, TArray<FVector>& PositionData);
 
 	/** Retrieves the Raw Rotation Animation Data for the given Animation Track Name and Animation Sequence */
 	UFUNCTION(BlueprintPure, Category = "AnimationBlueprintLibrary|RawTrackData")
-	static void GetRawTrackRotationData(const UAnimSequence* AnimationSequence, const FName TrackName, TArray<FQuat>& RotationData );
+	static void GetRawTrackRotationData(const UAnimSequenceBase* AnimationSequenceBase, const FName TrackName, TArray<FQuat>& RotationData );
 
 	/** Retrieves the Raw Scale Animation Data for the given Animation Track Name and Animation Sequence */
 	UFUNCTION(BlueprintPure, Category = "AnimationBlueprintLibrary|RawTrackData")
-	static void GetRawTrackScaleData(const UAnimSequence* AnimationSequence, const FName TrackName, TArray<FVector>& ScaleData);
+	static void GetRawTrackScaleData(const UAnimSequenceBase* AnimationSequenceBase, const FName TrackName, TArray<FVector>& ScaleData);
 
 	/** Retrieves the Raw Animation Data for the given Animation Track Name and Animation Sequence */
 	UFUNCTION(BlueprintPure, Category = "AnimationBlueprintLibrary|RawTrackData")
-	static void GetRawTrackData(const UAnimSequence* AnimationSequence, const FName TrackName, TArray<FVector>& PositionKeys,TArray<FQuat>& RotationKeys, TArray<FVector>& ScalingKeys);
+	static void GetRawTrackData(const UAnimSequenceBase* AnimationSequenceBase, const FName TrackName, TArray<FVector>& PositionKeys,TArray<FQuat>& RotationKeys, TArray<FVector>& ScalingKeys);
 
 	/** Checks whether or not the given Animation Track Name is contained within the Animation Sequence */
 	UFUNCTION(BlueprintCallable, Category = "AnimationBlueprintLibrary|Helpers")
-	static bool IsValidRawAnimationTrackName(const UAnimSequence* AnimationSequence, const FName TrackName);
+	static bool IsValidRawAnimationTrackName(const UAnimSequenceBase* AnimationSequenceBase, const FName TrackName);
 
-	static const FRawAnimSequenceTrack& GetRawAnimationTrackByName(const UAnimSequence* AnimationSequence, const FName TrackName);
+	static const FRawAnimSequenceTrack& GetRawAnimationTrackByName(const UAnimSequenceBase* AnimationSequenceBase, const FName TrackName);
 
 	// Compression
 
@@ -179,72 +179,72 @@ public:
 
 	/** Retrieves all Animation Notify Events found within the given Animation Sequence */
 	UFUNCTION(BlueprintPure, Category = "AnimationBlueprintLibrary|NotifyEvents")
-	static void GetAnimationNotifyEvents(const UAnimSequence* AnimationSequence, TArray<FAnimNotifyEvent>& NotifyEvents);
+	static void GetAnimationNotifyEvents(const UAnimSequenceBase* AnimationSequenceBase, TArray<FAnimNotifyEvent>& NotifyEvents);
 
 	/** Retrieves all Unique Animation Notify Events found within the given Animation Sequence */
 	UFUNCTION(BlueprintPure, Category = "AnimationBlueprintLibrary|NotifyEvents")
-	static void GetAnimationNotifyEventNames(const UAnimSequence* AnimationSequence, TArray<FName>& EventNames);
+	static void GetAnimationNotifyEventNames(const UAnimSequenceBase* AnimationSequenceBase, TArray<FName>& EventNames);
 
 	/** Adds an Animation Notify Event to Notify track in the given Animation with the given Notify creation data */
 	UFUNCTION(BlueprintCallable, Category = "AnimationBlueprintLibrary|NotifyEvents")
-	static UAnimNotify* AddAnimationNotifyEvent(UAnimSequence* AnimationSequence, FName NotifyTrackName, float StartTime, TSubclassOf<UAnimNotify> NotifyClass);
+	static UAnimNotify* AddAnimationNotifyEvent(UAnimSequenceBase* AnimationSequenceBase, FName NotifyTrackName, float StartTime, TSubclassOf<UAnimNotify> NotifyClass);
 
 	/** Adds an Animation Notify State Event to Notify track in the given Animation with the given Notify State creation data */
 	UFUNCTION(BlueprintCallable, Category = "AnimationBlueprintLibrary|NotifyEvents")
-	static UAnimNotifyState* AddAnimationNotifyStateEvent(UAnimSequence* AnimationSequence, FName NotifyTrackName, float StartTime, float Duration, TSubclassOf<UAnimNotifyState> NotifyStateClass);
+	static UAnimNotifyState* AddAnimationNotifyStateEvent(UAnimSequenceBase* AnimationSequenceBase, FName NotifyTrackName, float StartTime, float Duration, TSubclassOf<UAnimNotifyState> NotifyStateClass);
 
 	/** Adds an the specific Animation Notify to the Animation Sequence (requires Notify's outer to be the Animation Sequence) */
 	UFUNCTION(BlueprintCallable, Category = "AnimationBlueprintLibrary|NotifyEvents")
-	static void AddAnimationNotifyEventObject(UAnimSequence* AnimationSequence, float StartTime, UAnimNotify* Notify, FName NotifyTrackName);
+	static void AddAnimationNotifyEventObject(UAnimSequenceBase* AnimationSequenceBase, float StartTime, UAnimNotify* Notify, FName NotifyTrackName);
 
 	/** Adds an the specific Animation Notify State to the Animation Sequence (requires Notify State's outer to be the Animation Sequence) */
 	UFUNCTION(BlueprintCallable, Category = "AnimationBlueprintLibrary|NotifyEvents")
-	static void AddAnimationNotifyStateEventObject(UAnimSequence* AnimationSequence, float StartTime, float Duration, UAnimNotifyState* NotifyState, FName NotifyTrackName);
+	static void AddAnimationNotifyStateEventObject(UAnimSequenceBase* AnimationSequenceBase, float StartTime, float Duration, UAnimNotifyState* NotifyState, FName NotifyTrackName);
 
 	/** Removes Animation Notify Events found by Name within the Animation Sequence, and returns the number of removed name instances */
 	UFUNCTION(BlueprintCallable, Category = "AnimationBlueprintLibrary|NotifyEvents")
-	static int32 RemoveAnimationNotifyEventsByName(UAnimSequence* AnimationSequence, FName NotifyName);
+	static int32 RemoveAnimationNotifyEventsByName(UAnimSequenceBase* AnimationSequenceBase, FName NotifyName);
 
 	/** Removes Animation Notify Events found by Track within the Animation Sequence, and returns the number of removed name instances */
 	UFUNCTION(BlueprintCallable, Category = "AnimationBlueprintLibrary|NotifyEvents")
-	static int32 RemoveAnimationNotifyEventsByTrack(UAnimSequence* AnimationSequence, FName NotifyTrackName);	
+	static int32 RemoveAnimationNotifyEventsByTrack(UAnimSequenceBase* AnimationSequenceBase, FName NotifyTrackName);	
 
 	/** Replaces animation notifies in the specified Animation Sequence */
 	UFUNCTION(BlueprintCallable, Category = "AnimationBlueprintLibrary|NotifyEvents")
-	static void ReplaceAnimNotifyStates(UAnimSequenceBase* AnimationSequence, TSubclassOf<UAnimNotifyState> OldNotifyClass, TSubclassOf<UAnimNotifyState> NewNotifyClass, FOnNotifyStateReplaced OnNotifyStateReplaced);
+	static void ReplaceAnimNotifyStates(UAnimSequenceBase* AnimationSequenceBase, TSubclassOf<UAnimNotifyState> OldNotifyClass, TSubclassOf<UAnimNotifyState> NewNotifyClass, FOnNotifyStateReplaced OnNotifyStateReplaced);
 
 	/** Replaces animation notifies in the specified Animation Sequence */
 	UFUNCTION(BlueprintCallable, Category = "AnimationBlueprintLibrary|NotifyEvents")
-	static void ReplaceAnimNotifies(UAnimSequenceBase* AnimationSequence, TSubclassOf<UAnimNotify> OldNotifyClass, TSubclassOf<UAnimNotify> NewNotifyClass, FOnNotifyReplaced OnNotifyReplaced);
+	static void ReplaceAnimNotifies(UAnimSequenceBase* AnimationSequenceBase, TSubclassOf<UAnimNotify> OldNotifyClass, TSubclassOf<UAnimNotify> NewNotifyClass, FOnNotifyReplaced OnNotifyReplaced);
 
 	/** Copies animation notifies from Src Animation Sequence to Dest. Creates anim notify tracks as necessary. Returns true on success. */
 	UFUNCTION(BlueprintCallable, Category = "AnimationBlueprintLibrary|NotifyEvents")
-	static void CopyAnimNotifiesFromSequence(UAnimSequence* SrcAnimSequence, UAnimSequence* DestAnimSequence);
+	static void CopyAnimNotifiesFromSequence(UAnimSequenceBase* SourceAnimationSequenceBase, UAnimSequenceBase* DestinationAnimationSequenceBase);
 
 	// Notify Tracks
 
 	/** Retrieves all Unique Animation Notify Track Names found within the given Animation Sequence */
 	UFUNCTION(BlueprintPure, Category = "AnimationBlueprintLibrary|AnimationNotifies")
-	static void GetAnimationNotifyTrackNames(const UAnimSequence* AnimationSequence, TArray<FName>& TrackNames);
+	static void GetAnimationNotifyTrackNames(const UAnimSequenceBase* AnimationSequenceBase, TArray<FName>& TrackNames);
 
 	/** Adds an Animation Notify Track to the Animation Sequence */
 	UFUNCTION(BlueprintCallable, Category = "AnimationBlueprintLibrary|AnimationNotifies")
-	static void AddAnimationNotifyTrack(UAnimSequence* AnimationSequence, FName NotifyTrackName, FLinearColor TrackColor = FLinearColor::White);
+	static void AddAnimationNotifyTrack(UAnimSequenceBase* AnimationSequenceBase, FName NotifyTrackName, FLinearColor TrackColor = FLinearColor::White);
 
 	/** Removes an Animation Notify Track from Animation Sequence by Name */
 	UFUNCTION(BlueprintCallable, Category = "AnimationBlueprintLibrary|AnimationNotifies")
-	static void RemoveAnimationNotifyTrack(UAnimSequence* AnimationSequence, FName NotifyTrackName);
+	static void RemoveAnimationNotifyTrack(UAnimSequenceBase* AnimationSequenceBase, FName NotifyTrackName);
 
 	/** Removes All Animation Notify Tracks from Animation Sequence */
 	UFUNCTION(BlueprintCallable, Category = "AnimationBlueprintLibrary|AnimationNotifies")
-	static void RemoveAllAnimationNotifyTracks(UAnimSequence* AnimationSequence);
+	static void RemoveAllAnimationNotifyTracks(UAnimSequenceBase* AnimationSequenceBase);
 
 	/** Checks whether or not the given Track Name is a valid Animation Notify Track in the Animation Sequence */
 	UFUNCTION(BlueprintPure, Category = "AnimationBlueprintLibrary|Helpers")
-	static bool IsValidAnimNotifyTrackName(const UAnimSequence* AnimationSequence, FName NotifyTrackName);
+	static bool IsValidAnimNotifyTrackName(const UAnimSequenceBase* AnimationSequenceBase, FName NotifyTrackName);
 
-	static int32 GetTrackIndexForAnimationNotifyTrackName(const UAnimSequence* AnimationSequence, FName NotifyTrackName);
-	static const FAnimNotifyTrack& GetNotifyTrackByName(const UAnimSequence* AnimationSequence, FName NotifyTrackName);
+	static int32 GetTrackIndexForAnimationNotifyTrackName(const UAnimSequenceBase* AnimationSequenceBase, FName NotifyTrackName);
+	static const FAnimNotifyTrack& GetNotifyTrackByName(const UAnimSequenceBase* AnimationSequenceBase, FName NotifyTrackName);
 
 	/** Returns the actual trigger time for a NotifyEvent */
 	UFUNCTION(BlueprintPure, Category = "AnimationBlueprintLibrary|AnimationNotifies")
@@ -256,7 +256,7 @@ public:
 
 	/** Retrieves all Animation Notify Events for the given Notify Track Name from the given Animation Sequence */
 	UFUNCTION(BlueprintPure, Category = "AnimationBlueprintLibrary|NotifyEvents")
-	static void GetAnimationNotifyEventsForTrack(const UAnimSequence* AnimationSequence, FName NotifyTrackName, TArray<FAnimNotifyEvent>& Events);
+	static void GetAnimationNotifyEventsForTrack(const UAnimSequenceBase* AnimationSequenceBase, FName NotifyTrackName, TArray<FAnimNotifyEvent>& Events);
 
 	// Curves
 
@@ -363,55 +363,55 @@ public:
 
 	// MetaData
 
-	/** Creates and Adds an instance of the specified MetaData Class to the given Animation Sequence */
+	/** Creates and Adds an instance of the specified MetaData Class to the given Animation Asset */
 	UFUNCTION(BlueprintCallable, Category = "AnimationBlueprintLibrary|MetaData")
-	static void AddMetaData(UAnimSequence* AnimationSequence, TSubclassOf<UAnimMetaData> MetaDataClass, UAnimMetaData*& MetaDataInstance);
+	static void AddMetaData(UAnimationAsset* AnimationAsset, TSubclassOf<UAnimMetaData> MetaDataClass, UAnimMetaData*& MetaDataInstance);
 
-	/** Adds an instance of the specified MetaData Class to the given Animation Sequence (requires MetaDataObject's outer to be the Animation Sequence) */
+	/** Adds an instance of the specified MetaData Class to the given Animation Asset (requires MetaDataObject's outer to be the Animation Asset) */
 	UFUNCTION(BlueprintCallable, Category = "AnimationBlueprintLibrary|MetaData")
-	static void AddMetaDataObject(UAnimSequence* AnimationSequence, UAnimMetaData* MetaDataObject);
+	static void AddMetaDataObject(UAnimationAsset* AnimationAsset, UAnimMetaData* MetaDataObject);
 
-	/** Removes all Meta Data from the given Animation Sequence */
+	/** Removes all Meta Data from the given Animation Asset */
 	UFUNCTION(BlueprintCallable, Category = "AnimationBlueprintLibrary|MetaData")
-	static void RemoveAllMetaData(UAnimSequence* AnimationSequence);
+	static void RemoveAllMetaData(UAnimationAsset* AnimationAsset);
 
-	/** Removes the specified Meta Data Instance from the given Animation Sequence */
+	/** Removes the specified Meta Data Instance from the given Animation Asset */
 	UFUNCTION(BlueprintCallable, Category = "AnimationBlueprintLibrary|MetaData")
-	static void RemoveMetaData(UAnimSequence* AnimationSequence, UAnimMetaData* MetaDataObject);
+	static void RemoveMetaData(UAnimationAsset* AnimationAsset, UAnimMetaData* MetaDataObject);
 
-	/** Removes all Meta Data Instance of the specified Class from the given Animation Sequence */
+	/** Removes all Meta Data Instance of the specified Class from the given Animation Asset */
 	UFUNCTION(BlueprintCallable, Category = "AnimationBlueprintLibrary|MetaData")
-	static void RemoveMetaDataOfClass(UAnimSequence* AnimationSequence, TSubclassOf<UAnimMetaData> MetaDataClass);
+	static void RemoveMetaDataOfClass(UAnimationAsset* AnimationAsset, TSubclassOf<UAnimMetaData> MetaDataClass);
 
-	/** Retrieves all Meta Data Instances from the given Animation Sequence */
+	/** Retrieves all Meta Data Instances from the given Animation Asset */
 	UFUNCTION(BlueprintPure, Category = "AnimationBlueprintLibrary|MetaData")
-	static void GetMetaData(const UAnimSequence* AnimationSequence, TArray<UAnimMetaData*>& MetaData);
+	static void GetMetaData(const UAnimationAsset* AnimationAsset, TArray<UAnimMetaData*>& MetaData);
 
-	/** Retrieves all Meta Data Instances from the given Animation Sequence */
+	/** Retrieves all Meta Data Instances from the given Animation Asset */
 	UFUNCTION(BlueprintPure, Category = "AnimationBlueprintLibrary|MetaData")
-	static void GetMetaDataOfClass(const UAnimSequence* AnimationSequence, TSubclassOf<UAnimMetaData> MetaDataClass, TArray<UAnimMetaData*>& MetaDataOfClass);
+	static void GetMetaDataOfClass(const UAnimationAsset* AnimationAsset, TSubclassOf<UAnimMetaData> MetaDataClass, TArray<UAnimMetaData*>& MetaDataOfClass);
 
-	/** Checks whether or not the given Animation Sequences contains Meta Data Instance of the specified Meta Data Class */
+	/** Checks whether or not the given Animation Asset contains Meta Data Instance of the specified Meta Data Class */
 	UFUNCTION(BlueprintPure, Category = "AnimationBlueprintLibrary|MetaData")
-	static bool ContainsMetaDataOfClass(const UAnimSequence* AnimationSequence, TSubclassOf<UAnimMetaData> MetaDataClass);
+	static bool ContainsMetaDataOfClass(const UAnimationAsset* AnimationAsset, TSubclassOf<UAnimMetaData> MetaDataClass);
 
 	// Poses
 
 	/** Retrieves Bone Pose data for the given Bone Name at the specified Time from the given Animation Sequence */
 	UFUNCTION(BlueprintPure, Category = "AnimationBlueprintLibrary|Pose")
-	static void GetBonePoseForTime(const UAnimSequence* AnimationSequence, FName BoneName, float Time, bool bExtractRootMotion, FTransform& Pose);
+	static void GetBonePoseForTime(const UAnimSequenceBase* AnimationSequenceBase, FName BoneName, float Time, bool bExtractRootMotion, FTransform& Pose);
 
 	/** Retrieves Bone Pose data for the given Bone Name at the specified Frame from the given Animation Sequence */
 	UFUNCTION(BlueprintPure, Category = "AnimationBlueprintLibrary|Pose")
-	static void GetBonePoseForFrame(const UAnimSequence* AnimationSequence, FName BoneName, int32 Frame, bool bExtractRootMotion, FTransform& Pose);
+	static void GetBonePoseForFrame(const UAnimSequenceBase* AnimationSequenceBase, FName BoneName, int32 Frame, bool bExtractRootMotion, FTransform& Pose);
 
 	/** Retrieves Bone Pose data for the given Bone Names at the specified Time from the given Animation Sequence */
 	UFUNCTION(BlueprintPure, Category = "AnimationBlueprintLibrary|Pose")
-	static void GetBonePosesForTime(const UAnimSequence* AnimationSequence, TArray<FName> BoneNames, float Time, bool bExtractRootMotion, TArray<FTransform>& Poses, const USkeletalMesh* PreviewMesh = nullptr);
+	static void GetBonePosesForTime(const UAnimSequenceBase* AnimationSequenceBase, TArray<FName> BoneNames, float Time, bool bExtractRootMotion, TArray<FTransform>& Poses, const USkeletalMesh* PreviewMesh = nullptr);
 
 	/** Retrieves Bone Pose data for the given Bone Names at the specified Frame from the given Animation Sequence */
 	UFUNCTION(BlueprintPure, Category = "AnimationBlueprintLibrary|Pose")
-	static void GetBonePosesForFrame(const UAnimSequence* AnimationSequence, TArray<FName> BoneNames, int32 Frame, bool bExtractRootMotion, TArray<FTransform>& Poses, const USkeletalMesh* PreviewMesh = nullptr);
+	static void GetBonePosesForFrame(const UAnimSequenceBase* AnimationSequenceBase, TArray<FName> BoneNames, int32 Frame, bool bExtractRootMotion, TArray<FTransform>& Poses, const USkeletalMesh* PreviewMesh = nullptr);
 
 	// Virtual bones
 
@@ -437,33 +437,33 @@ public:
 
 	/** Retrieves the Length of the given Animation Sequence */
 	UFUNCTION(BlueprintPure, Category = "AnimationBlueprintLibrary|Animation")
-	static void GetSequenceLength(const UAnimSequence* AnimationSequence, float& Length);
+	static void GetSequenceLength(const UAnimSequenceBase* AnimationSequenceBase, float& Length);
 
 	/** Retrieves the (Play) Rate Scale of the given Animation Sequence */
 	UFUNCTION(BlueprintPure, Category = "AnimationBlueprintLibrary|Animation")
-	static void GetRateScale(const UAnimSequence* AnimationSequence, float& RateScale);
+	static void GetRateScale(const UAnimSequenceBase* AnimationSequenceBase, float& RateScale);
 
 	/** Sets the (Play) Rate Scale for the given Animation Sequence */
 	UFUNCTION(BlueprintCallable, Category = "AnimationBlueprintLibrary|Animation")
-	static void SetRateScale(UAnimSequence* AnimationSequence, float RateScale);
+	static void SetRateScale(UAnimSequenceBase* AnimationSequenceBase, float RateScale);
 
 	/** Retrieves the Frame Index at the specified Time Value for the given Animation Sequence */
 	UFUNCTION(BlueprintPure, Category = "AnimationBlueprintLibrary|Helpers")
-	static void GetFrameAtTime(const UAnimSequence* AnimationSequence, const float Time, int32& Frame);
+	static void GetFrameAtTime(const UAnimSequenceBase* AnimationSequenceBase, const float Time, int32& Frame);
 
 	/** Retrieves the Time Value at the specified Frame Indexfor the given Animation Sequence */
 	UFUNCTION(BlueprintPure, Category = "AnimationBlueprintLibrary|Helpers")
-	static void GetTimeAtFrame(const UAnimSequence* AnimationSequence, const int32 Frame, float& Time);
+	static void GetTimeAtFrame(const UAnimSequenceBase* AnimationSequenceBase, const int32 Frame, float& Time);
 	
-	static float GetTimeAtFrameInternal(const UAnimSequence* AnimationSequence, const int32 Frame);
+	static float GetTimeAtFrameInternal(const UAnimSequenceBase* AnimationSequenceBase, const int32 Frame);
 
 	/** Checks whether or not the given Time Value lies within the given Animation Sequence's Length */
 	UFUNCTION(BlueprintPure, Category = "AnimationBlueprintLibrary|Helpers")
-	static void IsValidTime(const UAnimSequence* AnimationSequence, const float Time, bool& IsValid);
+	static void IsValidTime(const UAnimSequenceBase* AnimationSequenceBase, const float Time, bool& IsValid);
 
-	static bool IsValidTimeInternal(const UAnimSequence* AnimationSequence, const float Time);
+	static bool IsValidTimeInternal(const UAnimSequenceBase* AnimationSequenceBase, const float Time);
 
 	/** Finds the Bone Path from the given Bone to the Root Bone */
 	UFUNCTION(BlueprintPure, Category = "AnimationBlueprintLibrary|Helpers")
-	static void FindBonePathToRoot(const UAnimSequence* AnimationSequence, FName BoneName, TArray<FName>& BonePath);
+	static void FindBonePathToRoot(const UAnimSequenceBase* AnimationSequenceBase, FName BoneName, TArray<FName>& BonePath);
 };
diff --git a/Engine/Source/Runtime/Engine/Classes/Animation/AnimSequenceHelpers.h b/Engine/Source/Runtime/Engine/Classes/Animation/AnimSequenceHelpers.h
--- a/Engine/Source/Runtime/Engine/Classes/Animation/AnimSequenceHelpers.h
+++ b/Engine/Source/Runtime/Engine/Classes/Animation/AnimSequenceHelpers.h
@@ -94,7 +94,7 @@ namespace Anim {
 	*
 	* @return Whether or not the copy was succesful
 	*/
-	ENGINE_API bool CopyNotifies(const UAnimSequence* SourceAnimSeq, UAnimSequence* DestAnimSeq, bool bShowDialogs = true);
+	ENGINE_API bool CopyNotifies(const UAnimSequenceBase* SourceAnimSeq, UAnimSequenceBase* DestAnimSeq, bool bShowDialogs = true);
 
 	namespace AnimationData
 	{
diff --git a/Engine/Source/Runtime/Engine/Private/Animation/AnimSequenceHelpers.cpp b/Engine/Source/Runtime/Engine/Private/Animation/AnimSequenceHelpers.cpp
--- a/Engine/Source/Runtime/Engine/Private/Animation/AnimSequenceHelpers.cpp
+++ b/Engine/Source/Runtime/Engine/Private/Animation/AnimSequenceHelpers.cpp
@@ -395,7 +395,7 @@ void Retargeting::RetargetPose(FCompactPose& InOutPose, const FName& RetargetSou
 
 
 #if WITH_EDITOR
-bool CopyNotifies(const UAnimSequence* SourceAnimSeq, UAnimSequence* DestAnimSeq, bool bShowDialogs)
+bool CopyNotifies(const UAnimSequenceBase* SourceAnimSeq, UAnimSequenceBase* DestAnimSeq, bool bShowDialogs)
 {
 	// Abort if source == destination.
 	if (SourceAnimSeq == DestAnimSeq)
