AnimationBlueprintLibrary does not assign a GUID to created notify events
#fix ensure that GUID is actually populated (previously did not break majorly due to PostSerialize fix-up)

diff --git a/Engine/Source/Editor/AnimationModifiers/Private/AnimationBlueprintLibrary.cpp b/Engine/Source/Editor/AnimationModifiers/Private/AnimationBlueprintLibrary.cpp
--- a/Engine/Source/Editor/AnimationModifiers/Private/AnimationBlueprintLibrary.cpp
+++ b/Engine/Source/Editor/AnimationModifiers/Private/AnimationBlueprintLibrary.cpp
@@ -595,6 +595,7 @@ UAnimNotify* UAnimationBlueprintLibrary::AddAnimationNotifyEvent(UAnimSequenceBa
 			NewEvent.TriggerTimeOffset = GetTriggerTimeOffsetForType(AnimationSequenceBase->CalculateOffsetForNotify(StartTime));
 			NewEvent.TrackIndex = GetTrackIndexForAnimationNotifyTrackName(AnimationSequenceBase, NotifyTrackName);
 			NewEvent.NotifyStateClass = nullptr;
+			NewEvent.Guid = FGuid::NewGuid();
 
 			if (NotifyClass)
 			{
@@ -653,6 +654,7 @@ UAnimNotifyState* UAnimationBlueprintLibrary::AddAnimationNotifyStateEvent(UAnim
 			NewEvent.TriggerTimeOffset = GetTriggerTimeOffsetForType(AnimationSequenceBase->CalculateOffsetForNotify(StartTime));
 			NewEvent.TrackIndex = GetTrackIndexForAnimationNotifyTrackName(AnimationSequenceBase, NotifyTrackName);
 			NewEvent.Notify = nullptr;
+			NewEvent.Guid = FGuid::NewGuid();
 
 			if (NotifyStateClass)
 			{
@@ -715,6 +717,7 @@ void UAnimationBlueprintLibrary::AddAnimationNotifyEventObject(UAnimSequenceBase
 			NewEvent.TrackIndex = GetTrackIndexForAnimationNotifyTrackName(AnimationSequenceBase, NotifyTrackName);
 			NewEvent.NotifyStateClass = nullptr;
 			NewEvent.Notify = Notify;
+			NewEvent.Guid = FGuid::NewGuid();		
 
 			// Refresh all cached data
 			AnimationSequenceBase->RefreshCacheData();
@@ -770,6 +773,7 @@ void UAnimationBlueprintLibrary::AddAnimationNotifyStateEventObject(UAnimSequenc
 			NewEvent.Notify = nullptr;
 			NewEvent.SetDuration(Duration);
 			NewEvent.EndLink.Link(AnimationSequenceBase, NewEvent.EndLink.GetTime());
+			NewEvent.Guid = FGuid::NewGuid();
 
 			// Refresh all cached data
 			AnimationSequenceBase->RefreshCacheData();
@@ -851,6 +855,7 @@ static void ReplaceAnimNotifies_Helper(UAnimSequenceBase* AnimationSequence, UCl
 					NewEvent.ChangeSlotIndex(SlotIndex);
 					NewEvent.SetSegmentIndex(SegmentIndex);
 					NewEvent.ChangeLinkMethod(LinkMethod);
+					NewEvent.Guid = FGuid::NewGuid();
 
 					UObject* AnimNotifyClass = NewObject<UObject>(AnimationSequence, NewNotifyClass, NAME_None, RF_Transactional);
 					NewEvent.NotifyStateClass = Cast<UAnimNotifyState>(AnimNotifyClass);
@@ -1033,7 +1038,7 @@ void UAnimationBlueprintLibrary::RemoveAnimationNotifyTrack(UAnimSequenceBase* A
 		{	
 			// Remove all notifies and sync markers on the to-delete-track
 			AnimationSequenceBase->Notifies.RemoveAll([&](const FAnimNotifyEvent& Notify) { return Notify.TrackIndex == TrackIndexToDelete; });
-		
+
 			// Before track removal, make sure everything behind is fixed
 			for (FAnimNotifyEvent& Notify : AnimationSequenceBase->Notifies)
 			{
@@ -1046,13 +1051,13 @@ void UAnimationBlueprintLibrary::RemoveAnimationNotifyTrack(UAnimSequenceBase* A
 			if (UAnimSequence* AnimationSequence = Cast<UAnimSequence>(AnimationSequenceBase))
 			{
 				AnimationSequence->AuthoredSyncMarkers.RemoveAll([&](const FAnimSyncMarker& Marker) { return Marker.TrackIndex == TrackIndexToDelete; });
-				for (FAnimSyncMarker& SyncMarker : AnimationSequence->AuthoredSyncMarkers)
+			for (FAnimSyncMarker& SyncMarker : AnimationSequence->AuthoredSyncMarkers)
+			{
+				if (SyncMarker.TrackIndex > TrackIndexToDelete)
 				{
-					if (SyncMarker.TrackIndex > TrackIndexToDelete)
-					{
-						SyncMarker.TrackIndex = SyncMarker.TrackIndex - 1;
-					}
+					SyncMarker.TrackIndex = SyncMarker.TrackIndex - 1;
 				}
+			}
 			}			
 			
 			// Delete the track itself
@@ -1071,11 +1076,11 @@ void UAnimationBlueprintLibrary::RemoveAnimationNotifyTrack(UAnimSequenceBase* A
 void UAnimationBlueprintLibrary::RemoveAllAnimationNotifyTracks(UAnimSequenceBase* AnimationSequenceBase)
 {
 	if (AnimationSequenceBase)
-	{
+{
 		AnimationSequenceBase->Notifies.Empty();
 		if (UAnimSequence* AnimationSequence = Cast<UAnimSequence>(AnimationSequenceBase))
-		{
-			AnimationSequence->AuthoredSyncMarkers.Empty();
+	{
+		AnimationSequence->AuthoredSyncMarkers.Empty();
 		}
 
 		// Remove all but one notify tracks
@@ -1861,7 +1866,7 @@ void UAnimationBlueprintLibrary::GetBonePosesForTime(const UAnimSequenceBase* An
 					if (IsValidRawAnimationTrackName(AnimationSequenceBase, BoneName))
 					{
 						const EAnimInterpolationType InterpolationType = [AnimationSequenceBase]() -> EAnimInterpolationType
-						{
+					{
 							if (const UAnimSequence* AnimationSequence = Cast<const UAnimSequence>(AnimationSequenceBase))
 							{
 								return AnimationSequence->Interpolation;
