- Added GetAnimNotifyEventDuration
- Set duration when replacing AnimNotify states

diff --git a/Engine/Source/Editor/AnimationModifiers/Private/AnimationBlueprintLibrary.cpp b/Engine/Source/Editor/AnimationModifiers/Private/AnimationBlueprintLibrary.cpp
--- a/Engine/Source/Editor/AnimationModifiers/Private/AnimationBlueprintLibrary.cpp
+++ b/Engine/Source/Editor/AnimationModifiers/Private/AnimationBlueprintLibrary.cpp
@@ -865,6 +865,7 @@ static void ReplaceAnimNotifies_Helper(UAnimSequenceBase* AnimationSequence, UCl
 					if (NewEvent.NotifyStateClass)
 					{
 						NewEvent.NotifyName = FName(*NewEvent.NotifyStateClass->GetNotifyName());
+						NewEvent.SetDuration(Length);
 						NewEvent.EndTriggerTimeOffset = EndTriggerTimeOffset;
 						NewEvent.EndLink.ChangeSlotIndex(EndSlotIndex);
 						NewEvent.EndLink.SetSegmentIndex(EndSegmentIndex);
@@ -1135,6 +1136,11 @@ float UAnimationBlueprintLibrary::GetAnimNotifyEventTriggerTime(const FAnimNotif
 	return NotifyEvent.GetTriggerTime();
 }
 
+float UAnimationBlueprintLibrary::GetAnimNotifyEventDuration(const FAnimNotifyEvent& NotifyEvent)
+{
+	return NotifyEvent.Duration;
+}
+
 void UAnimationBlueprintLibrary::GetAnimationSyncMarkersForTrack(const UAnimSequence* AnimationSequence, FName NotifyTrackName, TArray<FAnimSyncMarker>& Markers)
 {
 	Markers.Empty();
diff --git a/Engine/Source/Editor/AnimationModifiers/Public/AnimationBlueprintLibrary.h b/Engine/Source/Editor/AnimationModifiers/Public/AnimationBlueprintLibrary.h
--- a/Engine/Source/Editor/AnimationModifiers/Public/AnimationBlueprintLibrary.h
+++ b/Engine/Source/Editor/AnimationModifiers/Public/AnimationBlueprintLibrary.h
@@ -249,6 +249,10 @@ public:
 	/** Returns the actual trigger time for a NotifyEvent */
 	UFUNCTION(BlueprintPure, Category = "AnimationBlueprintLibrary|AnimationNotifies")
 	static float GetAnimNotifyEventTriggerTime(const FAnimNotifyEvent& NotifyEvent);
+	
+	/** Returns the duration for a NotifyEvent, only non-zero for Anim Notify States */
+	UFUNCTION(BlueprintPure, Category = "AnimationBlueprintLibrary|AnimationNotifies")
+	static float GetAnimNotifyEventDuration(const FAnimNotifyEvent& NotifyEvent);
 
 	/** Retrieves all Animation Sync Markers for the given Notify Track Name from the given Animation Sequence */
 	UFUNCTION(BlueprintPure, Category = "AnimationBlueprintLibrary|MarkerSyncing")
