//UE5/Main - Cook EngineTest Win64 - Assertion failed: !FUObjectThreadContext::Get().IsRoutingPostLoad
#fix update revision GUID for any AnimModifier class its instance during load, rather than forcefully loading them during the CDOs PostLoad

diff --git a/Engine/Source/Editor/AnimationModifiers/Private/AnimationModifier.cpp b/Engine/Source/Editor/AnimationModifiers/Private/AnimationModifier.cpp
--- a/Engine/Source/Editor/AnimationModifiers/Private/AnimationModifier.cpp
+++ b/Engine/Source/Editor/AnimationModifiers/Private/AnimationModifier.cpp
@@ -188,12 +188,25 @@ void UAnimationModifier::Serialize(FArchive& Ar)
 
 void UAnimationModifier::PostLoad()
 {
-	Super::PostLoad();	
-	// Ensure we always have a valid guid
-	if (!RevisionGuid.IsValid())
+	Super::PostLoad();
+
+	UClass* Class = GetClass();
+	UObject* DefaultObject = Class->GetDefaultObject();
+
+	// CDO, set GUID if invalid
+	if(DefaultObject == this)
+	{
+		// Ensure we always have a valid guid
+		if (!RevisionGuid.IsValid())
+		{
+			UpdateRevisionGuid(GetClass());
+			MarkPackageDirty();
+		}
+	}
+	// Non CDO, update revision GUID
+	else if(UAnimationModifier* TypedDefaultObject = Cast<UAnimationModifier>(DefaultObject))
 	{
-		UpdateRevisionGuid(GetClass());
-		MarkPackageDirty();
+		RevisionGuid = TypedDefaultObject->RevisionGuid;
 	}
 }
 
@@ -206,11 +219,9 @@ void UAnimationModifier::UpdateRevisionGuid(UClass* ModifierClass)
 {
 	if (ModifierClass)
 	{
-		LoadModifierReferencers(ModifierClass);
-	
 		RevisionGuid = FGuid::NewGuid();
 
-		// Native classes are more difficult?
+		// Apply to any currently loaded instances of this class
 		for (TObjectIterator<UAnimationModifier> It; It; ++It)
 		{
 			if (*It != this && It->GetClass() == ModifierClass)
