Fix no-unity compile

#rnx

diff --git a/Engine/Source/Editor/AnimationModifiers/Private/AnimationModifier.cpp b/Engine/Source/Editor/AnimationModifiers/Private/AnimationModifier.cpp
--- a/Engine/Source/Editor/AnimationModifiers/Private/AnimationModifier.cpp
+++ b/Engine/Source/Editor/AnimationModifiers/Private/AnimationModifier.cpp
@@ -5,6 +5,7 @@
 #include "Animation/AnimSequence.h"
 #include "Animation/Skeleton.h"
 #include "ModifierOutputFilter.h"
+#include "ScopedTransaction.h"
 #include "AssetRegistry/AssetRegistryModule.h"
 #include "Editor/Transactor.h"
 #include "UObject/UObjectIterator.h"
diff --git a/Engine/Source/Editor/AnimationModifiers/Private/AnimationModifiersModule.cpp b/Engine/Source/Editor/AnimationModifiers/Private/AnimationModifiersModule.cpp
--- a/Engine/Source/Editor/AnimationModifiers/Private/AnimationModifiersModule.cpp
+++ b/Engine/Source/Editor/AnimationModifiers/Private/AnimationModifiersModule.cpp
@@ -13,6 +13,7 @@
 #include "Modules/ModuleManager.h"
 #include "PropertyEditorModule.h" 
 #include "SAnimationModifierContentBrowserWindow.h"
+#include "ScopedTransaction.h" 
 #include "Framework/Application/SlateApplication.h"
 #include "Interfaces/IMainFrameModule.h"
 
