Feature Request: Apply All for Animations modified by [REDACTED_PERSON]
#feature added ApplyModifier scope struct which keeps track of which modifier has prompted the user for any warnings, and allow them to YesAll/NoAll rather than having to individually prompt for each modifier/animation sequence combination

diff --git a/Engine/Source/Editor/AnimationModifiers/Private/AnimationModifier.cpp b/Engine/Source/Editor/AnimationModifiers/Private/AnimationModifier.cpp
--- a/Engine/Source/Editor/AnimationModifiers/Private/AnimationModifier.cpp
+++ b/Engine/Source/Editor/AnimationModifiers/Private/AnimationModifier.cpp
@@ -18,6 +18,22 @@
 
 #define LOCTEXT_NAMESPACE "AnimationModifier"
 
+int32 UE::Anim::FApplyModifiersScope::ScopesOpened = 0;
+TMap<FObjectKey, TOptional<EAppReturnType::Type>>  UE::Anim::FApplyModifiersScope::PerClassReturnTypeValues;
+
+TOptional<EAppReturnType::Type> UE::Anim::FApplyModifiersScope::GetReturnType(const UAnimationModifier* InModifier)
+{
+	TOptional<EAppReturnType::Type>* ReturnTypePtr = PerClassReturnTypeValues.Find(FObjectKey(InModifier->GetClass()));
+	return ReturnTypePtr ? *ReturnTypePtr : TOptional<EAppReturnType::Type>();
+}
+
+void UE::Anim::FApplyModifiersScope::SetReturnType(const UAnimationModifier* InModifier, EAppReturnType::Type InReturnType)
+{
+	const FObjectKey Key(InModifier->GetClass());
+	ensure(!PerClassReturnTypeValues.Contains(Key));
+	PerClassReturnTypeValues.Add(Key, InReturnType);
+}
+
 UAnimationModifier::UAnimationModifier()
 	: PreviouslyAppliedModifier(nullptr)
 {
@@ -35,6 +51,7 @@ void UAnimationModifier::ApplyToAnimationSequence(class UAnimSequence* InAnimati
 	FCategoryLogOutputFilter OutputLog;
 	OutputLog.SetAutoEmitLineTerminator(true);
 	OutputLog.AddCategoryName("LogAnimationBlueprintLibrary");
+	OutputLog.AddCategoryName("LogAnimation");
 
 	GLog->AddOutputDevice(&OutputLog);
 		
@@ -69,17 +86,50 @@ void UAnimationModifier::ApplyToAnimationSequence(class UAnimSequence* InAnimati
 	GLog->RemoveOutputDevice(&OutputLog);
 
 	// Check if warnings or errors have occurred and show dialog to user to inform her about this
-	const bool bWarningsOrErrors = OutputLog.ContainsWarnings() || OutputLog.ContainsErrors();
-	bool bShouldRevert = false;
-	if (bWarningsOrErrors)
+	const bool bWarnings = OutputLog.ContainsWarnings();
+	const bool bErrors = OutputLog.ContainsErrors();
+	bool bShouldRevert = bErrors;
+	
+	// If errors have occured - prompt user with OK and revert
+	TOptional<EAppReturnType::Type> ScopeReturnType = UE::Anim::FApplyModifiersScope::GetReturnType(this);	
+	static const FText MessageTitle = LOCTEXT("ModifierDialogTitle", "Modifier has generated errors during test run.");
+	if (bErrors)
+	{		
+		if (UE::Anim::FApplyModifiersScope::ScopesOpened == 0 || !ScopeReturnType.IsSet() || ScopeReturnType.GetValue() != EAppReturnType::Type::Ok)
+		{
+			const FText ErrorMessageFormat = FText::FormatOrdered(LOCTEXT("ModifierErrorDescription", "Modifier: {0}\nAsset: {1}\n{2}\nResolve the errors before trying to apply again."), FText::FromString(GetClass()->GetPathName()),
+				FText::FromString(CurrentAnimSequence ? CurrentAnimSequence->GetPathName() : CurrentSkeleton->GetPathName()), FText::FromString(OutputLog));
+			
+			EAppReturnType::Type ReturnValue = FMessageDialog::Open(EAppMsgType::Ok, ErrorMessageFormat, &MessageTitle);
+			UE::Anim::FApplyModifiersScope::SetReturnType(this, ReturnValue);
+		}
+	}
+
+	// If _only_ warnings have occured - check if user has previously said YesAll / NoAll and process the result
+	if (bWarnings && !bShouldRevert)
 	{
-		static const FText WarningMessageFormat = FText::FromString("Modifier has generated warnings during a test run:\n\n{0}\nAre you sure you want to Apply it?");
-		static const FText ErrorMessageFormat = FText::FromString("Modifier has generated errors (and warnings) during a test run:\n\n{0}\nResolve the Errors before trying to Apply!");
+		if (UE::Anim::FApplyModifiersScope::ScopesOpened == 0 || !ScopeReturnType.IsSet())
+		{
+			const FText WarningMessage = FText::FormatOrdered(LOCTEXT("ModifierWarningDescription", "Modifier: {0}\nAsset: {1}\n{2}\nAre you sure you want to apply it?"), FText::FromString(GetClass()->GetPathName()),
+				FText::FromString(CurrentAnimSequence ? CurrentAnimSequence->GetPathName() : CurrentSkeleton->GetPathName()), FText::FromString(OutputLog));
+
+			EAppReturnType::Type ReturnValue = FMessageDialog::Open(EAppMsgType::YesNoYesAllNoAll, WarningMessage, &MessageTitle);
+			bShouldRevert = ReturnValue == EAppReturnType::No || ReturnValue == EAppReturnType::NoAll;
 
-		EAppMsgType::Type MessageType = OutputLog.ContainsErrors() ? EAppMsgType::Ok : EAppMsgType::YesNo;
-		const FText& MessageFormat = OutputLog.ContainsErrors() ? ErrorMessageFormat : WarningMessageFormat;
-		const FText MessageTitle = FText::FromString("Modifier has Generated Warnings/Errors");
-		bShouldRevert = (FMessageDialog::Open(MessageType, FText::FormatOrdered(MessageFormat, FText::FromString(OutputLog)), &MessageTitle) != EAppReturnType::Yes);
+			// check if user response should be stored for further modifier applications
+			if(UE::Anim::FApplyModifiersScope::ScopesOpened > 0)
+			{
+				if (ReturnValue == EAppReturnType::Type::YesAll || ReturnValue == EAppReturnType::Type::NoAll)
+				{
+					UE::Anim::FApplyModifiersScope::SetReturnType(this, ReturnValue);
+				}
+			}
+		}
+		else
+		{
+			// Revert if previous user prompt return NoAll or if any errors occured previously 
+			bShouldRevert = ScopeReturnType.GetValue() == EAppReturnType::NoAll || ScopeReturnType.GetValue() == EAppReturnType::Ok;
+		}
 	}
 
 	// Revert changes if necessary, otherwise post edit and refresh animation data
diff --git a/Engine/Source/Editor/AnimationModifiers/Private/AnimationModifiersModule.cpp b/Engine/Source/Editor/AnimationModifiers/Private/AnimationModifiersModule.cpp
--- a/Engine/Source/Editor/AnimationModifiers/Private/AnimationModifiersModule.cpp
+++ b/Engine/Source/Editor/AnimationModifiers/Private/AnimationModifiersModule.cpp
@@ -103,6 +103,7 @@ void FAnimationModifiersModule::ApplyAnimationModifiers(const TArray<UAnimSequen
 	const FScopedTransaction Transaction(LOCTEXT("UndoAction_ApplyModifiers", "Applying Animation Modifier(s) to Animation Sequence(s)"));
 	
 	// Iterate over each Animation Sequence and all of its contained modifiers, applying each one
+	UE::Anim::FApplyModifiersScope Scope;
 	TArray<UAnimationModifiersAssetUserData*> AssetUserData;
 	for (UAnimSequence* AnimationSequence : InSequences)
 	{
diff --git a/Engine/Source/Editor/AnimationModifiers/Private/SAnimationModifiersTab.cpp b/Engine/Source/Editor/AnimationModifiers/Private/SAnimationModifiersTab.cpp
--- a/Engine/Source/Editor/AnimationModifiers/Private/SAnimationModifiersTab.cpp
+++ b/Engine/Source/Editor/AnimationModifiers/Private/SAnimationModifiersTab.cpp
@@ -446,6 +446,7 @@ void SAnimationModifiersTab::ApplyModifiers(const TArray<UAnimationModifier*>& M
 
 	if (bApply)
 	{
+		UE::Anim::FApplyModifiersScope Scope;
 		for (UAnimSequence* AnimSequence : AnimSequences)
 		{
 			AnimSequence->Modify();
diff --git a/Engine/Source/Editor/AnimationModifiers/Public/AnimationModifier.h b/Engine/Source/Editor/AnimationModifiers/Public/AnimationModifier.h
--- a/Engine/Source/Editor/AnimationModifiers/Public/AnimationModifier.h
+++ b/Engine/Source/Editor/AnimationModifiers/Public/AnimationModifier.h
@@ -4,6 +4,7 @@
 
 #include "Animation/AnimEnums.h"
 #include "Animation/AnimCurveTypes.h"
+#include "Misc/MessageDialog.h"
 #include "AnimationBlueprintLibrary.h"
 
 #include "AnimationModifier.generated.h"
@@ -84,4 +85,38 @@ private:
 	/** Serialized version of the modifier that has been previously applied to the Animation Asset */
 	UPROPERTY()
 	TObjectPtr<UAnimationModifier> PreviouslyAppliedModifier;
-};
\ No newline at end of file
+};
+
+namespace UE
+{
+	namespace Anim
+	{
+		struct FApplyModifiersScope
+		{
+			FApplyModifiersScope()
+			{
+				if (ScopesOpened == 0)
+				{				
+					PerClassReturnTypeValues.Empty();
+				}
+				++ScopesOpened;
+			}
+
+			~FApplyModifiersScope()
+			{
+				--ScopesOpened;
+				check(ScopesOpened >= 0);
+				if(ScopesOpened == 0)
+				{
+					PerClassReturnTypeValues.Empty();
+				}
+			}
+
+			static TOptional<EAppReturnType::Type> GetReturnType(const UAnimationModifier* InModifier);			
+			static void SetReturnType(const class UAnimationModifier* InModifier, EAppReturnType::Type InReturnType);
+
+			static TMap<FObjectKey, TOptional<EAppReturnType::Type>> PerClassReturnTypeValues;
+			static int32 ScopesOpened;
+		};
+	}
+}
