Option to apply Animation Modifiers to only new or changed animations
#feature add ability to only apply out-of-date animation modifiers from asset context menu, and through details view on AnimationModifier classes
#fix ensure that we do not apply pending kill modifiers or instances that are used for reverting

diff --git a/Engine/Source/Developer/AssetTools/Private/AssetTypeActions/AssetTypeActions_AnimSequence.cpp b/Engine/Source/Developer/AssetTools/Private/AssetTypeActions/AssetTypeActions_AnimSequence.cpp
--- a/Engine/Source/Developer/AssetTools/Private/AssetTypeActions/AssetTypeActions_AnimSequence.cpp
+++ b/Engine/Source/Developer/AssetTools/Private/AssetTypeActions/AssetTypeActions_AnimSequence.cpp
@@ -99,6 +99,30 @@ void FAssetTypeActions_AnimSequence::GetActions(const TArray<UObject*>& InObject
 				}
 	        }))
 	    );
+
+		MenuBuilder.AddMenuEntry(
+			LOCTEXT("AnimSequence_ApplyOutOfDataAnimationModifier", "Apply out-of-date Modifiers"),
+			LOCTEXT("AnimSequence_ApplyOutOfDataAnimationModifierTooltip", "Applies all contained animation modifier(s), if they are out of date."),
+			FSlateIcon(FEditorStyle::GetStyleSetName(), "ClassIcon.AnimationModifier"),
+			FUIAction(FExecuteAction::CreateLambda([Sequences]()
+			{
+			TArray<UAnimSequence*> AnimSequences;
+			Algo::TransformIf(Sequences, AnimSequences, 
+			[](const TWeakObjectPtr<UAnimSequence>& WeakAnimSequence)
+			{
+			    return WeakAnimSequence.Get() && WeakAnimSequence->IsA<UAnimSequence>();
+			},
+			[](const TWeakObjectPtr<UAnimSequence>& WeakAnimSequence)
+			{
+			    return WeakAnimSequence.Get();
+			});
+
+			if (IAnimationModifiersModule* Module = FModuleManager::Get().LoadModulePtr<IAnimationModifiersModule>("AnimationModifiers"))
+			{
+			    Module->ApplyAnimationModifiers(AnimSequences, false);
+			}
+			}))
+	    );
 	});
 	
 	Section.AddSubMenu("AnimSequence_AnimationModifiers", LOCTEXT("AnimSequence_AnimationModifiers", "Animation Modifier(s)"),
diff --git a/Engine/Source/Editor/AnimationModifiers/Private/AnimationModifier.cpp b/Engine/Source/Editor/AnimationModifiers/Private/AnimationModifier.cpp
--- a/Engine/Source/Editor/AnimationModifiers/Private/AnimationModifier.cpp
+++ b/Engine/Source/Editor/AnimationModifiers/Private/AnimationModifier.cpp
@@ -34,6 +34,8 @@ void UE::Anim::FApplyModifiersScope::SetReturnType(const UAnimationModifier* InM
 	PerClassReturnTypeValues.Add(Key, InReturnType);
 }
 
+const FName UAnimationModifier::RevertModifierObjectName("REVERT_AnimationModifier");
+
 UAnimationModifier::UAnimationModifier()
 	: PreviouslyAppliedModifier(nullptr)
 {
@@ -43,8 +45,8 @@ void UAnimationModifier::ApplyToAnimationSequence(class UAnimSequence* InAnimati
 {
 	FEditorScriptExecutionGuard ScriptGuard;
 
-	checkf(InAnimationSequence, TEXT("Invalid Animation Sequence supplied"));
 	CurrentAnimSequence = InAnimationSequence;
+	checkf(CurrentAnimSequence, TEXT("Invalid Animation Sequence supplied"));
 	CurrentSkeleton = InAnimationSequence->GetSkeleton();
 
 	// Filter to check for warnings / errors thrown from animation blueprint library (rudimentary approach for now)
@@ -148,7 +150,7 @@ void UAnimationModifier::ApplyToAnimationSequence(class UAnimSequence* InAnimati
 			PreviouslyAppliedModifier->MarkPendingKill();
 		}
 
-		PreviouslyAppliedModifier = DuplicateObject(this, GetOuter());
+		PreviouslyAppliedModifier = DuplicateObject(this, GetOuter(), RevertModifierObjectName);
 
 		CurrentAnimSequence->PostEditChange();
 		CurrentSkeleton->PostEditChange();
@@ -305,7 +307,7 @@ void UAnimationModifier::UpdateNativeRevisionGuid()
 	}
 }
 
-void UAnimationModifier::ApplyToAll(TSubclassOf<UAnimationModifier> ModifierSubClass)
+void UAnimationModifier::ApplyToAll(TSubclassOf<UAnimationModifier> ModifierSubClass, bool bForceApply /*= true*/)
 {
 	if (UClass* ModifierClass = ModifierSubClass.Get())
 	{
@@ -315,20 +317,25 @@ void UAnimationModifier::ApplyToAll(TSubclassOf<UAnimationModifier> ModifierSubC
 		const FScopedTransaction Transaction(LOCTEXT("UndoAction_ApplyModifiers", "Applying Animation Modifier to Animation Sequence(s)"));		
 		for (TObjectIterator<UAnimationModifier> It; It; ++It)
 		{
-			if (*It && It->GetClass() == ModifierClass)
+			// Check if valid, of the required class, not pending kill and not a modifier back-up for reverting
+			if (*It && It->GetClass() == ModifierClass && !It->IsPendingKill() && It->GetFName() != RevertModifierObjectName)
 			{
-				// Go through outer chain to find AnimSequence
-				UObject* Outer = It->GetOuter();
-				while(Outer && !Outer->IsA<UAnimSequence>())
+				if (bForceApply || !It->IsLatestRevisionApplied())
 				{
-					Outer = Outer->GetOuter();
+					// Go through outer chain to find AnimSequence
+					UObject* Outer = It->GetOuter();
+
+					while(Outer && !Outer->IsA<UAnimSequence>())
+					{
+						Outer = Outer->GetOuter();
+					}
+
+					if (UAnimSequence* AnimSequence = Cast<UAnimSequence>(Outer))
+					{
+						AnimSequence->Modify();
+						It->ApplyToAnimationSequence(AnimSequence);
+					}	
 				}
-
-				if (UAnimSequence* AnimSequence = Cast<UAnimSequence>(Outer))
-				{
-					AnimSequence->Modify();
-					It->ApplyToAnimationSequence(AnimSequence);
-				}			
 			}
 		}
 	}	
diff --git a/Engine/Source/Editor/AnimationModifiers/Private/AnimationModifierDetailCustomization.cpp b/Engine/Source/Editor/AnimationModifiers/Private/AnimationModifierDetailCustomization.cpp
--- a/Engine/Source/Editor/AnimationModifiers/Private/AnimationModifierDetailCustomization.cpp
+++ b/Engine/Source/Editor/AnimationModifiers/Private/AnimationModifierDetailCustomization.cpp
@@ -35,14 +35,23 @@ void FAnimationModifierDetailCustomization::CustomizeDetails(IDetailLayoutBuilde
 		[
 			SNew(SHorizontalBox)
 			+SHorizontalBox::Slot()
-           .AutoWidth()
-           .Padding(2.f)
-           [
-               SNew(SButton)
-               .OnClicked(this, &FAnimationModifierDetailCustomization::OnApplyButtonClicked)
-               .ToolTipText(LOCTEXT("ApplyToolTip", "Applies any instanced modifiers of this class to their owning Animation Sequences."))
-               .Text(LOCTEXT("ApplyText", "Apply to All"))
-           ]
+			.AutoWidth()
+			.Padding(2.f)
+			[
+			   SNew(SButton)
+			   .OnClicked(this, &FAnimationModifierDetailCustomization::OnApplyButtonClicked, true)
+			   .ToolTipText(LOCTEXT("ApplyToolTip", "Applies any instanced modifiers of this class to their owning Animation Sequences."))
+			   .Text(LOCTEXT("ApplyText", "Apply to All"))
+			]
+			+SHorizontalBox::Slot()
+			.AutoWidth()
+			.Padding(2.f)
+			[
+				SNew(SButton)
+				.OnClicked(this, &FAnimationModifierDetailCustomization::OnApplyButtonClicked, false)
+				.ToolTipText(LOCTEXT("ApplyOutOfDataToolTip", "Applies any instanced modifiers, if they are out-of-date, of this class to their owning Animation Sequences."))
+				.Text(LOCTEXT("ApplyToOutOfDateText", "Apply to All out-of-date"))
+			]
 			+SHorizontalBox::Slot()
 			.AutoWidth()
 			.Padding(2.f)
@@ -51,7 +60,7 @@ void FAnimationModifierDetailCustomization::CustomizeDetails(IDetailLayoutBuilde
 				.OnClicked(this, &FAnimationModifierDetailCustomization::OnUpdateRevisionButtonClicked)
 				.ToolTipText(LOCTEXT("UpdateRevisionToolTip", "Updates the stored revision GUID on all instances of this Modifier class, marking them out-of-date."))
 				.Text(LOCTEXT("UpdateRevisionText", "Update Revision"))
-            ]           
+			]           
 		];
 	}
 }
@@ -65,11 +74,11 @@ FReply FAnimationModifierDetailCustomization::OnUpdateRevisionButtonClicked()
 	return FReply::Handled();
 }
 
-FReply FAnimationModifierDetailCustomization::OnApplyButtonClicked()
+FReply FAnimationModifierDetailCustomization::OnApplyButtonClicked(bool bForceApply)
 {
 	if (ModifierInstance)
 	{
-		UAnimationModifier::ApplyToAll(ModifierInstance->GetClass());
+		UAnimationModifier::ApplyToAll(ModifierInstance->GetClass(), bForceApply);
 	}
 	return FReply::Handled();
 }
diff --git a/Engine/Source/Editor/AnimationModifiers/Private/AnimationModifierDetailCustomization.h b/Engine/Source/Editor/AnimationModifiers/Private/AnimationModifierDetailCustomization.h
--- a/Engine/Source/Editor/AnimationModifiers/Private/AnimationModifierDetailCustomization.h
+++ b/Engine/Source/Editor/AnimationModifiers/Private/AnimationModifierDetailCustomization.h
@@ -16,7 +16,7 @@ public:
 	virtual void CustomizeDetails(IDetailLayoutBuilder& DetailBuilder) override;
 protected:
 	FReply OnUpdateRevisionButtonClicked();
-	FReply OnApplyButtonClicked();
+	FReply OnApplyButtonClicked(bool bForceApply);
 protected:
 	TSharedPtr<SButton> UpdateRevisionButton;
 	UAnimationModifier* ModifierInstance;
diff --git a/Engine/Source/Editor/AnimationModifiers/Private/AnimationModifiersModule.cpp b/Engine/Source/Editor/AnimationModifiers/Private/AnimationModifiersModule.cpp
--- a/Engine/Source/Editor/AnimationModifiers/Private/AnimationModifiersModule.cpp
+++ b/Engine/Source/Editor/AnimationModifiers/Private/AnimationModifiersModule.cpp
@@ -98,7 +98,7 @@ void FAnimationModifiersModule::ShowAddAnimationModifierWindow(const TArray<UAni
 	FSlateApplication::Get().AddModalWindow(Window, ParentWindow, false);
 }
 
-void FAnimationModifiersModule::ApplyAnimationModifiers(const TArray<UAnimSequence*>& InSequences)
+void FAnimationModifiersModule::ApplyAnimationModifiers(const TArray<UAnimSequence*>& InSequences, bool bForceApply /*= true*/)
 {
 	const FScopedTransaction Transaction(LOCTEXT("UndoAction_ApplyModifiers", "Applying Animation Modifier(s) to Animation Sequence(s)"));
 	
@@ -116,7 +116,10 @@ void FAnimationModifiersModule::ApplyAnimationModifiers(const TArray<UAnimSequen
 				const TArray<UAnimationModifier*>& ModifierInstances = UserData->GetAnimationModifierInstances();
 				for (UAnimationModifier* Modifier : ModifierInstances)
 				{
-					Modifier->ApplyToAnimationSequence(AnimationSequence);
+					if (bForceApply || !Modifier->IsLatestRevisionApplied())
+					{
+						Modifier->ApplyToAnimationSequence(AnimationSequence);
+					}
 				}
 			}
 		}		
diff --git a/Engine/Source/Editor/AnimationModifiers/Private/AnimationModifiersModule.h b/Engine/Source/Editor/AnimationModifiers/Private/AnimationModifiersModule.h
--- a/Engine/Source/Editor/AnimationModifiers/Private/AnimationModifiersModule.h
+++ b/Engine/Source/Editor/AnimationModifiers/Private/AnimationModifiersModule.h
@@ -20,7 +20,7 @@ public:
 
 	/** Begin IAnimationModifiersModule overrides */
 	virtual void ShowAddAnimationModifierWindow(const TArray<UAnimSequence*>& InSequences) override;
-	virtual void ApplyAnimationModifiers(const TArray<UAnimSequence*>& InSequences) override;
+	virtual void ApplyAnimationModifiers(const TArray<UAnimSequence*>& InSequences, bool bForceApply = true) override;
 	/** End IAnimationModifiersModule overrides */
 
 protected:
diff --git a/Engine/Source/Editor/AnimationModifiers/Public/AnimationModifier.h b/Engine/Source/Editor/AnimationModifiers/Public/AnimationModifier.h
--- a/Engine/Source/Editor/AnimationModifiers/Public/AnimationModifier.h
+++ b/Engine/Source/Editor/AnimationModifiers/Public/AnimationModifier.h
@@ -60,7 +60,7 @@ protected:
 	void UpdateNativeRevisionGuid();
 
 	/** Applies all instances of the provided Modifier class to its outer Animation Sequence*/
-	static void ApplyToAll(TSubclassOf<UAnimationModifier> ModifierSubClass);
+	static void ApplyToAll(TSubclassOf<UAnimationModifier> ModifierSubClass, bool bForceApply = true);
 	static void LoadModifierReferencers(TSubclassOf<UAnimationModifier> ModifierSubClass);
 private:
 	UAnimSequence* CurrentAnimSequence;
@@ -85,6 +85,8 @@ private:
 	/** Serialized version of the modifier that has been previously applied to the Animation Asset */
 	UPROPERTY()
 	TObjectPtr<UAnimationModifier> PreviouslyAppliedModifier;
+
+	static const FName RevertModifierObjectName;
 };
 
 namespace UE
diff --git a/Engine/Source/Editor/AnimationModifiers/Public/IAnimationModifiersModule.h b/Engine/Source/Editor/AnimationModifiers/Public/IAnimationModifiersModule.h
--- a/Engine/Source/Editor/AnimationModifiers/Public/IAnimationModifiersModule.h
+++ b/Engine/Source/Editor/AnimationModifiers/Public/IAnimationModifiersModule.h
@@ -14,5 +14,5 @@ public:
 	virtual void ShowAddAnimationModifierWindow(const TArray<UAnimSequence*>& InSequences) = 0;
 
 	/** Apply animation modifiers contained by provided Animation Sequences to its respective outers */
-	virtual void ApplyAnimationModifiers(const TArray<UAnimSequence*>& InSequences) = 0;	
+	virtual void ApplyAnimationModifiers(const TArray<UAnimSequence*>& InSequences, bool bForceApply = true) = 0;
 };
