Fixed Static analysis issues in interchange (fbxmesh.h) by changing the function GetMeshUniqueID from FbxHelper to return a empty string when the mesh is empty.

Fixed Static analysis in AnimationModifier.cpp the CurrentAnimSequence ptr can't be null since it was deferenced at the line 75.

Fixed Static analysis in PackageStoreOptimizer for some reason the static analysis did like the way the check was presented.

Fixed Static analysis in SComponentClassCombo.cpp and SDisplayClusterConfiguratorComponentCombo.cpp changed the code so that the static analyser should understand it better.

diff --git a/Engine/Plugins/Runtime/Interchange/Source/Parsers/Fbx/Private/FbxHelper.cpp b/Engine/Plugins/Runtime/Interchange/Source/Parsers/Fbx/Private/FbxHelper.cpp
--- a/Engine/Plugins/Runtime/Interchange/Source/Parsers/Fbx/Private/FbxHelper.cpp
+++ b/Engine/Plugins/Runtime/Interchange/Source/Parsers/Fbx/Private/FbxHelper.cpp
@@ -42,6 +42,11 @@ namespace UE
 
 			FString FFbxHelper::GetMeshUniqueID(FbxGeometryBase* Mesh)
 			{
+				if (!Mesh)
+				{
+					return {};
+				}
+
 				FString MeshUniqueID;
 				if (Mesh->GetAttributeType() == FbxNodeAttribute::eMesh)
 				{
diff --git a/Engine/Plugins/Runtime/nDisplay/Source/DisplayClusterConfigurator/Private/Views/SCSEditor/SDisplayClusterConfiguratorComponentCombo.cpp b/Engine/Plugins/Runtime/nDisplay/Source/DisplayClusterConfigurator/Private/Views/SCSEditor/SDisplayClusterConfiguratorComponentCombo.cpp
--- a/Engine/Plugins/Runtime/nDisplay/Source/DisplayClusterConfigurator/Private/Views/SCSEditor/SDisplayClusterConfiguratorComponentCombo.cpp
+++ b/Engine/Plugins/Runtime/nDisplay/Source/DisplayClusterConfigurator/Private/Views/SCSEditor/SDisplayClusterConfiguratorComponentCombo.cpp
@@ -178,7 +178,7 @@ void SDisplayClusterConfiguratorComponentClassCombo::GenerateFilteredComponentLi
 				if ( TextFilter->TestTextFilter(FBasicStringFilterExpressionContext(FriendlyComponentName)) )
 				{
 					// Add the heading first if it hasn't already been added
-					if (LastHeadingIndex != INDEX_NONE)
+					if (LastHeadingPtr && LastHeadingIndex != INDEX_NONE)
 					{
 						FilteredComponentClassList.Insert(*LastHeadingPtr, LastHeadingIndex);
 						LastHeadingIndex = INDEX_NONE;
diff --git a/Engine/Source/Developer/IoStoreUtilities/Private/PackageStoreOptimizer.cpp b/Engine/Source/Developer/IoStoreUtilities/Private/PackageStoreOptimizer.cpp
--- a/Engine/Source/Developer/IoStoreUtilities/Private/PackageStoreOptimizer.cpp
+++ b/Engine/Source/Developer/IoStoreUtilities/Private/PackageStoreOptimizer.cpp
@@ -817,8 +817,8 @@ void FPackageStoreOptimizer::OptimizeExportBundles(const TMap<FPackageId, FPacka
 	FPackageStorePackage* PreviousPackage = nullptr;
 	for (FPackageStorePackage::FExportBundleGraphNode* Node : LoadOrder)
 	{
-		check(Node->Package);
 		FPackageStorePackage* Package = Node->Package;
+		check(Package);
 		if (PreviousPackage != Package || !Package->CurrentBundle)
 		{
 			if (Package->GraphData.ExportBundles.IsEmpty())
diff --git a/Engine/Source/Editor/AnimationModifiers/Private/AnimationModifier.cpp b/Engine/Source/Editor/AnimationModifiers/Private/AnimationModifier.cpp
--- a/Engine/Source/Editor/AnimationModifiers/Private/AnimationModifier.cpp
+++ b/Engine/Source/Editor/AnimationModifiers/Private/AnimationModifier.cpp
@@ -100,7 +100,7 @@ void UAnimationModifier::ApplyToAnimationSequence(class UAnimSequence* InAnimati
 		if (UE::Anim::FApplyModifiersScope::ScopesOpened == 0 || !ScopeReturnType.IsSet() || ScopeReturnType.GetValue() != EAppReturnType::Type::Ok)
 		{
 			const FText ErrorMessageFormat = FText::FormatOrdered(LOCTEXT("ModifierErrorDescription", "Modifier: {0}\nAsset: {1}\n{2}\nResolve the errors before trying to apply again."), FText::FromString(GetClass()->GetPathName()),
-				FText::FromString(CurrentAnimSequence ? CurrentAnimSequence->GetPathName() : CurrentSkeleton->GetPathName()), FText::FromString(OutputLog));
+				FText::FromString(CurrentAnimSequence->GetPathName()), FText::FromString(OutputLog));
 			
 			EAppReturnType::Type ReturnValue = FMessageDialog::Open(EAppMsgType::Ok, ErrorMessageFormat, &MessageTitle);
 			UE::Anim::FApplyModifiersScope::SetReturnType(this, ReturnValue);
diff --git a/Engine/Source/Editor/UnrealEd/Private/SComponentClassCombo.cpp b/Engine/Source/Editor/UnrealEd/Private/SComponentClassCombo.cpp
--- a/Engine/Source/Editor/UnrealEd/Private/SComponentClassCombo.cpp
+++ b/Engine/Source/Editor/UnrealEd/Private/SComponentClassCombo.cpp
@@ -158,7 +158,7 @@ void SComponentClassCombo::GenerateFilteredComponentList()
 				if ( TextFilter->TestTextFilter(FBasicStringFilterExpressionContext(FriendlyComponentName)) )
 				{
 					// Add the heading first if it hasn't already been added
-					if (LastHeadingIndex != INDEX_NONE)
+					if (LastHeadingPtr && LastHeadingIndex != INDEX_NONE)
 					{
 						FilteredComponentClassList.Insert(*LastHeadingPtr, LastHeadingIndex);
 						LastHeadingIndex = INDEX_NONE;
