PR #9136: Adding a setting for animation modifiers to reapply post owner change (Contributed by [REDACTED_PERSON])

diff --git a/Engine/Source/Editor/AnimationModifiers/Private/AnimationModifiersAssetUserData.cpp b/Engine/Source/Editor/AnimationModifiers/Private/AnimationModifiersAssetUserData.cpp
--- a/Engine/Source/Editor/AnimationModifiers/Private/AnimationModifiersAssetUserData.cpp
+++ b/Engine/Source/Editor/AnimationModifiers/Private/AnimationModifiersAssetUserData.cpp
@@ -1,11 +1,14 @@
 // Copyright Epic Games, Inc. All Rights Reserved.
 
 #include "AnimationModifiersAssetUserData.h"
+#include "AnimationModifier.h"
+#include "UObject/UObjectThreadContext.h"
 
 #include "HAL/PlatformCrt.h"
 #include "Math/UnrealMathUtility.h"
 #include "Misc/AssertionMacros.h"
 #include "UObject/ObjectPtr.h"
+#include "Animation/AnimSequence.h"
 
 void UAnimationModifiersAssetUserData::AddAnimationModifier(UAnimationModifier* Instance)
 {
@@ -40,6 +43,35 @@ void UAnimationModifiersAssetUserData::PostEditChangeProperty(struct FPropertyCh
 	Super::PostEditChangeProperty(PropertyChangedEvent);
 }
 
+#if WITH_EDITOR
+void UAnimationModifiersAssetUserData::PostEditChangeOwner()
+{
+	Super::PostEditChangeOwner();
+
+	// We cant call blueprint implemented functions while routing post load
+	if (FUObjectThreadContext::Get().IsRoutingPostLoad)
+	{
+		return;
+	}
+
+	UAnimSequence* ModifiedAnimSequence = Cast<UAnimSequence>(GetOuter());
+	if (!ModifiedAnimSequence)
+	{
+		return;
+	}
+
+	for (UAnimationModifier* Modifier : AnimationModifierInstances)
+	{
+		if (!Modifier->bReapplyPostOwnerChange || Modifier->IsCurrentlyApplyingModifier())
+		{
+			continue;
+		}
+
+		Modifier->ApplyToAnimationSequence(ModifiedAnimSequence);
+	}
+}
+#endif // WITH_EDITOR
+
 void UAnimationModifiersAssetUserData::Serialize(FArchive& Ar)
 {
 	Super::Serialize(Ar);
diff --git a/Engine/Source/Editor/AnimationModifiers/Public/AnimationModifier.h b/Engine/Source/Editor/AnimationModifiers/Public/AnimationModifier.h
--- a/Engine/Source/Editor/AnimationModifiers/Public/AnimationModifier.h
+++ b/Engine/Source/Editor/AnimationModifiers/Public/AnimationModifier.h
@@ -61,6 +61,14 @@ public:
 	virtual void Serialize(FArchive& Ar) override;
 	virtual void PostLoad() override;
 	// End UObject Overrides
+
+	/** If this is set to true then the animation modifier will call it's reapply function after any change made to the owning asset. */
+	UPROPERTY(EditAnywhere, Category = "Settings")
+	bool bReapplyPostOwnerChange = false;
+
+	/** Whether or not this modifier is in the process of being applied to an Animation Asset */
+	bool IsCurrentlyApplyingModifier() const { return CurrentAnimSequence != nullptr || CurrentSkeleton != nullptr; };
+
 protected:
 	// Derived class accessors to skeleton and anim sequence 
 	const UAnimSequence* GetAnimationSequence();
diff --git a/Engine/Source/Editor/AnimationModifiers/Public/AnimationModifiersAssetUserData.h b/Engine/Source/Editor/AnimationModifiers/Public/AnimationModifiersAssetUserData.h
--- a/Engine/Source/Editor/AnimationModifiers/Public/AnimationModifiersAssetUserData.h
+++ b/Engine/Source/Editor/AnimationModifiers/Public/AnimationModifiersAssetUserData.h
@@ -32,6 +32,9 @@ protected:
 	virtual void PostLoad() override;
 	virtual void PostEditChangeProperty(struct FPropertyChangedEvent& PropertyChangedEvent) override;	
 	virtual void Serialize(FArchive& Ar) override;
+#if WITH_EDITOR
+	virtual void PostEditChangeOwner() override;
+#endif // WITH_EDITOR
 	/** End UAssetUserData overrides */
 
 	void AddAnimationModifier(UAnimationModifier* Instance);
