Non-unity build fix
#rnx

diff --git a/Engine/Source/Editor/AnimationModifiers/Private/AnimationModifiersModule.h b/Engine/Source/Editor/AnimationModifiers/Private/AnimationModifiersModule.h
--- a/Engine/Source/Editor/AnimationModifiers/Private/AnimationModifiersModule.h
+++ b/Engine/Source/Editor/AnimationModifiers/Private/AnimationModifiersModule.h
@@ -7,6 +7,7 @@
 #include "Templates/SharedPointer.h"
 #include "UObject/NameTypes.h"
 #include "WorkflowOrientedApp/WorkflowCentricApplication.h"
+#include "AnimationModifier.h"
 
 class FApplicationMode;
 class UAnimSequence;
