Contributed by [REDACTED_PERSON] from The Coalition:

Animation Modifier : Fix Modifier on Skeleton
Transfered the ownership of 'PreviouslyAppliedModifier' from an animation modifier instance to the animation sequence being applied.
This solved the following issue:
- Modifier on Skeleton cause USkeleton dirtied everytime the modifier is applied to an animation sequence.
- Modifier on Skeleton cannot be re-apply or reverted correctly.
- CanRevert & OutOfDate status for Modifier on Skeleton was not reflect the true status of all animation sequences referencing that skeleton.

- CurrentAnimSequence/CurrentSkeleton was not set on OnRevert()
- IAnimationDataController::FScopedBracket was not open on OnRevert() before re-apply modifier
- Stateful animation modifier can now be reverted correctly (Applied modifier instance is nolonger reverted after OnApply call)

diff --git a/Engine/Source/Editor/AnimationModifiers/Private/AnimationModifier.cpp b/Engine/Source/Editor/AnimationModifiers/Private/AnimationModifier.cpp
--- a/Engine/Source/Editor/AnimationModifiers/Private/AnimationModifier.cpp
+++ b/Engine/Source/Editor/AnimationModifiers/Private/AnimationModifier.cpp
@@ -1,6 +1,8 @@
 // Copyright Epic Games, Inc. All Rights Reserved.
 
 #include "AnimationModifier.h"
+#include "AnimationModifiersAssetUserData.h"
+#include "AnimationModifierHelpers.h"
 
 #include "Algo/Transform.h"
 #include "Animation/AnimData/IAnimationDataController.h"
@@ -36,6 +38,8 @@
 int32 UE::Anim::FApplyModifiersScope::ScopesOpened = 0;
 TMap<FObjectKey, TOptional<EAppReturnType::Type>>  UE::Anim::FApplyModifiersScope::PerClassReturnTypeValues;
 
+const FName UAnimationModifier::AnimationModifiersTag = TEXT("AnimationModifierList");
+
 TOptional<EAppReturnType::Type> UE::Anim::FApplyModifiersScope::GetReturnType(const UAnimationModifier* InModifier)
 {
 	TOptional<EAppReturnType::Type>* ReturnTypePtr = PerClassReturnTypeValues.Find(FObjectKey(InModifier->GetClass()));
@@ -49,21 +53,18 @@ void UE::Anim::FApplyModifiersScope::SetReturnType(const UAnimationModifier* InM
 	PerClassReturnTypeValues.Add(Key, InReturnType);
 }
 
-const FName UAnimationModifier::RevertModifierObjectName("REVERT_AnimationModifier");
 
 UAnimationModifier::UAnimationModifier()
-	: PreviouslyAppliedModifier(nullptr)
 {
 }
 
-void UAnimationModifier::ApplyToAnimationSequence(class UAnimSequence* InAnimationSequence)
+void UAnimationModifier::ApplyToAnimationSequence(class UAnimSequence* AnimSequence) const
 {
 	FEditorScriptExecutionGuard ScriptGuard;
 
-	CurrentAnimSequence = InAnimationSequence;
-	checkf(CurrentAnimSequence, TEXT("Invalid Animation Sequence supplied"));
-	CurrentSkeleton = InAnimationSequence->GetSkeleton();
-	checkf(CurrentSkeleton, TEXT("Invalid Skeleton for supplied Animation Sequence"));
+	checkf(AnimSequence, TEXT("Invalid Animation Sequence supplied"));
+	USkeleton* Skeleton = AnimSequence->GetSkeleton();
+	checkf(Skeleton, TEXT("Invalid Skeleton for supplied Animation Sequence"));
 
 	// Filter to check for warnings / errors thrown from animation blueprint library (rudimentary approach for now)
 	FCategoryLogOutputFilter OutputLog;
@@ -73,30 +74,39 @@ void UAnimationModifier::ApplyToAnimationSequence(class UAnimSequence* InAnimati
 
 	GLog->AddOutputDevice(&OutputLog);
 		
-	// Transact the modifier to prevent instance variables/data to change during applying
-	FTransaction ModifierTransaction;
-	ModifierTransaction.SaveObject(this);
+	bool AppliedModifierOuterWasCreated = false;
+	UObject* AppliedModifierOuter = AnimSequence->GetAssetUserData<UAnimationModifiersAssetUserData>();
+	if (!AppliedModifierOuter)
+	{
+		AppliedModifierOuterWasCreated = true;
+		AppliedModifierOuter = FAnimationModifierHelpers::RetrieveOrCreateModifierUserData(AnimSequence);
+	}
+	// Using explicit name may cause duplicate path name
+	UAnimationModifier* AppliedModifier = DuplicateObject(this, AppliedModifierOuter);
+	// Set the revision guid on applied modifier to latest
+	AppliedModifier->RevisionGuid = GetLatestRevisionGuid();
 
 	FTransaction AnimationDataTransaction;
-	AnimationDataTransaction.SaveObject(CurrentAnimSequence);
-	AnimationDataTransaction.SaveObject(CurrentSkeleton);
+	AnimationDataTransaction.SaveObject(AnimSequence);
+	AnimationDataTransaction.SaveObject(Skeleton);
 
-	/** In case this modifier has been previously applied, revert it using the serialised out version at the time */	
-	if (PreviouslyAppliedModifier)
 	{
-		PreviouslyAppliedModifier->Modify();
-		PreviouslyAppliedModifier->OnRevert(CurrentAnimSequence);
-	}
+		// Group the OnRevert & OnApply operation into one Bracket to reduce compression request
+		IAnimationDataController::FScopedBracket ScopedBracket(AnimSequence->GetController(), LOCTEXT("ApplyModifierBracket", "Applying Animation Modifier"));
 
-	IAnimationDataController& Controller = CurrentAnimSequence->GetController(); //-V595
-
-	/** Reverting and applying, populates the log with possible warnings and or errors to notify the user about */
-	OnApply(CurrentAnimSequence);
+		/** In case this modifier has been previously applied, revert it using the serialised out version at the time */
+		if (UAnimationModifier* PreviouslyAppliedModifier = GetAppliedModifier(AnimSequence))
+		{
+			// Save the applied modifier as well
+			AnimationDataTransaction.SaveObject(PreviouslyAppliedModifier);
+			PreviouslyAppliedModifier->ExecuteOnRevert(AnimSequence);
+		}
 
-	// Apply transaction
-	ModifierTransaction.BeginOperation();
-	ModifierTransaction.Apply();
-	ModifierTransaction.EndOperation();
+		{
+			/** Reverting and applying, populates the log with possible warnings and or errors to notify the user about */
+			AppliedModifier->ExecuteOnApply(AnimSequence);
+		}
+	}
 
 	GLog->RemoveOutputDevice(&OutputLog);
 
@@ -113,7 +123,7 @@ void UAnimationModifier::ApplyToAnimationSequence(class UAnimSequence* InAnimati
 		if (UE::Anim::FApplyModifiersScope::ScopesOpened == 0 || !ScopeReturnType.IsSet() || ScopeReturnType.GetValue() != EAppReturnType::Type::Ok)
 		{
 			const FText ErrorMessageFormat = FText::FormatOrdered(LOCTEXT("ModifierErrorDescription", "Modifier: {0}\nAsset: {1}\n{2}\nResolve the errors before trying to apply again."), FText::FromString(GetClass()->GetPathName()),
-				FText::FromString(CurrentAnimSequence->GetPathName()), FText::FromString(OutputLog));
+				FText::FromString(AnimSequence->GetPathName()), FText::FromString(OutputLog));
 			
 			EAppReturnType::Type ReturnValue = FMessageDialog::Open(EAppMsgType::Ok, ErrorMessageFormat, &MessageTitle);
 			UE::Anim::FApplyModifiersScope::SetReturnType(this, ReturnValue);
@@ -126,7 +136,7 @@ void UAnimationModifier::ApplyToAnimationSequence(class UAnimSequence* InAnimati
 		if (UE::Anim::FApplyModifiersScope::ScopesOpened == 0 || !ScopeReturnType.IsSet())
 		{
 			const FText WarningMessage = FText::FormatOrdered(LOCTEXT("ModifierWarningDescription", "Modifier: {0}\nAsset: {1}\n{2}\nAre you sure you want to apply it?"), FText::FromString(GetClass()->GetPathName()),
-			FText::FromString(CurrentAnimSequence->GetPathName()), FText::FromString(OutputLog));
+			FText::FromString(AnimSequence->GetPathName()), FText::FromString(OutputLog));
 
 			EAppReturnType::Type ReturnValue = FMessageDialog::Open(EAppMsgType::YesNoYesAllNoAll, WarningMessage, &MessageTitle);
 			bShouldRevert = ReturnValue == EAppReturnType::No || ReturnValue == EAppReturnType::NoAll;
@@ -153,28 +163,28 @@ void UAnimationModifier::ApplyToAnimationSequence(class UAnimSequence* InAnimati
 		AnimationDataTransaction.BeginOperation();
 		AnimationDataTransaction.Apply();
 		AnimationDataTransaction.EndOperation();
-		CurrentAnimSequence->RefreshCacheData();
+		AnimSequence->RefreshCacheData();
+		AppliedModifier->MarkAsGarbage();
+		if (AppliedModifierOuterWasCreated)
+		{
+			// The animation sequence's asset user data array should have been reverted by the transaction
+			AppliedModifierOuter->MarkAsGarbage();
+		}
 	}
 	else
 	{
-		/** Mark the previous modifier pending kill, as it will be replaced with the current modifier state */
-		if (PreviouslyAppliedModifier)
+		SetAppliedModifier(AnimSequence, AppliedModifier);
+
+		if (Skeleton->GetPackage()->IsDirty())
 		{
-			PreviouslyAppliedModifier->MarkAsGarbage();
+			Skeleton->PostEditChange();
+		}
+		if (AnimSequence->GetPackage()->IsDirty())
+		{
+			AnimSequence->PostEditChange();
+			AnimSequence->RefreshCacheData();
 		}
-
-		PreviouslyAppliedModifier = DuplicateObject(this, GetOuter(), MakeUniqueObjectName(GetOuter(), GetClass(), RevertModifierObjectName));
-
-		CurrentAnimSequence->PostEditChange();
-		CurrentSkeleton->PostEditChange();
-		CurrentAnimSequence->RefreshCacheData();
-
-		UpdateStoredRevisions();
 	}
-		
-	// Finished
-	CurrentAnimSequence = nullptr;
-	CurrentSkeleton = nullptr;
 }
 
 void UAnimationModifier::UpdateCompressedAnimationData()
@@ -185,53 +195,183 @@ void UAnimationModifier::UpdateCompressedAnimationData()
 	}
 }
 
-void UAnimationModifier::RevertFromAnimationSequence(class UAnimSequence* InAnimationSequence)
+void UAnimationModifier::RevertFromAnimationSequence(class UAnimSequence* AnimSequence) const
 {
 	FEditorScriptExecutionGuard ScriptGuard;
 
+	checkf(AnimSequence, TEXT("Invalid Animation Sequence supplied"));
+	USkeleton* Skeleton = AnimSequence->GetSkeleton();
+	checkf(Skeleton, TEXT("Invalid Skeleton for supplied Animation Sequence"));
 	/** Can only revert if previously applied, which means there should be a previous modifier */
-	if (PreviouslyAppliedModifier)
+	UAnimationModifier* PreviouslyAppliedModifier = FindAndRemoveAppliedModifier(AnimSequence);
+
+	// Backward compatibility
+	// Is PreviouslyAppliedModifier owned by InAnimationSequence
+	// Modifier on Skeleton read from legacy version may _share_ the PreviouslyAppliedModifier
+	bool AppliedModifierOwnedByAnimation = true;
+	if (!PreviouslyAppliedModifier)
 	{
-		checkf(InAnimationSequence, TEXT("Invalid Animation Sequence supplied"));
-		CurrentAnimSequence = InAnimationSequence;
-		CurrentSkeleton = InAnimationSequence->GetSkeleton();
+		PreviouslyAppliedModifier = GetLegacyPreviouslyAppliedModifierForModifierOnSkeleton();
+		if (PreviouslyAppliedModifier)
+		{
+			AppliedModifierOwnedByAnimation = false;
+		}
+	}
 
+	if (PreviouslyAppliedModifier)
+	{
 		// Transact the modifier to prevent instance variables/data to change during reverting
-		FTransaction Transaction;
-		Transaction.SaveObject(this);
+		{
+			IAnimationDataController::FScopedBracket ScopedBracket(AnimSequence->GetController(), LOCTEXT("ApplyModifierBracket", "Applying Animation Modifier"));	
+			PreviouslyAppliedModifier->ExecuteOnRevert(AnimSequence);
+		}
 
-		PreviouslyAppliedModifier->Modify();
+		if (AppliedModifierOwnedByAnimation)
+		{
+			PreviouslyAppliedModifier->MarkAsGarbage();
+		}
 
-		IAnimationDataController& Controller = CurrentAnimSequence->GetController();
+		if (Skeleton->GetPackage()->IsDirty())
+		{
+			Skeleton->PostEditChange();
+		}
+		if (AnimSequence->GetPackage()->IsDirty())
+		{
+			AnimSequence->PostEditChange();
+			AnimSequence->RefreshCacheData();
+		}
+	}
+}
 
+bool UAnimationModifier::CanRevert(IInterface_AssetUserData* AssetUserDataInterface) const
+{
+	return GetAppliedModifier(AssetUserDataInterface) != nullptr;
+}
+
+UAnimationModifier* UAnimationModifier::GetAppliedModifier(IInterface_AssetUserData* AssetUserDataInterface) const
+{
+	if (UAnimationModifiersAssetUserData* AssetData = AssetUserDataInterface->GetAssetUserData<UAnimationModifiersAssetUserData>())
+	{
+		TObjectPtr<UAnimationModifier>* AppliedModifier = AssetData->AppliedModifiers.Find(this);
+		if (AppliedModifier)
 		{
-			IAnimationDataController::FScopedBracket ScopedBracket(Controller, LOCTEXT("RevertModifierBracket", "Reverting Animation Modifier"));
-			PreviouslyAppliedModifier->OnRevert(CurrentAnimSequence);
+			return AppliedModifier->Get();
 		}
+	}
 
-		// Apply transaction
-		Transaction.BeginOperation();
-		Transaction.Apply();
-		Transaction.EndOperation();
+	// Backward compatibility
+	// Read the shared applied instance for modifier on skeleton ready from legacy version
+	if (UAnimationModifier* LegacyAppliedModifierOnSkeleton = GetLegacyPreviouslyAppliedModifierForModifierOnSkeleton())
+	{
+		return LegacyAppliedModifierOnSkeleton;
+	}
 
-	    CurrentAnimSequence->PostEditChange();
-	    CurrentSkeleton->PostEditChange();
-	    CurrentAnimSequence->RefreshCacheData();
+	return nullptr;
+}
 
-		ResetStoredRevisions();
+UAnimationModifier* UAnimationModifier::FindAndRemoveAppliedModifier(TScriptInterface<IInterface_AssetUserData> AssetUserDataInterface) const
+{
+	if (UAnimationModifiersAssetUserData* AssetData = AssetUserDataInterface->GetAssetUserData<UAnimationModifiersAssetUserData>())
+	{
+		TObjectPtr<UAnimationModifier> AppliedModifier = nullptr;
+		AssetData->AppliedModifiers.RemoveAndCopyValue(this, AppliedModifier);
+		if (AppliedModifier) 
+		{
+			AssetData->Modify();
+			AssetUserDataInterface.GetObject()->Modify();
+		}
+		return AppliedModifier.Get();
+	}
+	return nullptr;
+}
 
-		// Finished
-		CurrentAnimSequence = nullptr;
-		CurrentSkeleton = nullptr;
+void UAnimationModifier::SetAppliedModifier(TScriptInterface<IInterface_AssetUserData> AssetUserDataInterface, UAnimationModifier* AppliedModifier) const
+{
+	if (UAnimationModifiersAssetUserData* AssetData = FAnimationModifierHelpers::RetrieveOrCreateModifierUserData(AssetUserDataInterface))
+	{
+		TObjectPtr<UAnimationModifier>& Modifier = AssetData->AppliedModifiers.FindOrAdd(this);
 
-		PreviouslyAppliedModifier->MarkAsGarbage();
-		PreviouslyAppliedModifier = nullptr;
+		if (Modifier)
+		{
+			Modifier->MarkAsGarbage();
+		}
+		Modifier = AppliedModifier;
+		AssetData->Modify();
+		AssetUserDataInterface.GetObject()->Modify();
 	}
 }
 
-bool UAnimationModifier::IsLatestRevisionApplied() const
+bool UAnimationModifier::IsLatestRevisionApplied(IInterface_AssetUserData* AssetUserDataInterface) const
 {
-	return (AppliedGuid == RevisionGuid);
+	if (UAnimationModifier* AppliedModifier = GetAppliedModifier(AssetUserDataInterface))
+	{
+		return AppliedModifier->RevisionGuid == GetLatestRevisionGuid();
+	}
+	return false;
+}
+
+UAnimationModifier* UAnimationModifier::GetLegacyPreviouslyAppliedModifierForModifierOnSkeleton() const
+{
+	// Read PreviouslyAppliedModifier_DEPRECATED from previous version
+	// Check PostLoad() when we moved the deprecated value into this place
+	if (UAnimationModifiersAssetUserData* AssetData = Cast<UAnimationModifiersAssetUserData>(GetOuter()))
+	{
+		if (USkeleton* Skeleton = Cast<USkeleton>(AssetData->GetOuter()))
+		{
+			if (TObjectPtr<UAnimationModifier>* PtrToModifier = AssetData->AppliedModifiers.Find(this))
+			{
+				// The new system uses a _placeholder_ object with _exact_ UAnimationAsset class to store RevisionGuid
+				// Any object with concrete modifier class must come from PreviouslyAppliedModifier_DEPRECATED
+				if (PtrToModifier->GetClass() != UAnimationAsset::StaticClass())
+				{
+					return PtrToModifier->Get();
+				}
+			}
+		}
+	}
+	return nullptr;
+}
+
+FGuid UAnimationModifier::GetLatestRevisionGuid() const
+{
+	return GetDefault<UAnimationModifier>(GetClass())->RevisionGuid;
+}
+
+void UAnimationModifier::GetAssetRegistryTagsForAppliedModifiersFromSkeleton(const UObject* Object, TArray<UObject::FAssetRegistryTag>& OutTags)
+{
+	FString ModifiersList;
+	if (UAnimSequence* AnimSequence = const_cast<UAnimSequence*>(Cast<UAnimSequence>(Object)))
+	{
+		if (USkeleton* Skeleton = AnimSequence->GetSkeleton())
+		{
+			if (UAnimationModifiersAssetUserData* SkeletonAssetData = Skeleton->GetAssetUserData<UAnimationModifiersAssetUserData>())
+			{
+				for (const UAnimationModifier* Modifier : SkeletonAssetData->GetAnimationModifierInstances())
+				{
+					if (const UAnimationModifier* AppliedModifier = Modifier->GetAppliedModifier(AnimSequence))
+					{
+						// "{Name}={Revision};"
+						FString Revision;
+						AppliedModifier->RevisionGuid.AppendString(Revision, EGuidFormats::Short);
+						ModifiersList += FString::Printf(TEXT("%s%c%s%c"), *Modifier->GetName(), AnimationModifiersAssignment, *Revision, AnimationModifiersDelimiter);
+					}
+				}
+				if (!ModifiersList.IsEmpty())
+				{
+					OutTags.Emplace(AnimationModifiersTag, ModifiersList, UObject::FAssetRegistryTag::TT_Hidden);
+				}
+			}
+		}
+	}
+}
+
+//! @brief Mark this modifier on skeleton as reverted (affect CanRevert, IsLatestRevisionApplied)
+void UAnimationModifier::RemoveLegacyPreviousAppliedModifierOnSkeleton(USkeleton* Skeleton)
+{
+	if (UAnimationModifier* AppliedModifier = FindAndRemoveAppliedModifier(Skeleton))
+	{
+		AppliedModifier->MarkAsGarbage();
+	}
 }
 
 void UAnimationModifier::PostInitProperties()
@@ -248,7 +388,13 @@ void UAnimationModifier::Serialize(FArchive& Ar)
 	/** Back-wards compatibility, assume the current modifier as previously applied */
 	if (Ar.CustomVer(FReleaseObjectVersion::GUID) < FReleaseObjectVersion::SerializeAnimModifierState)
 	{
-		PreviouslyAppliedModifier = DuplicateObject(this, GetOuter());
+		UObject* AssetData = GetOuter();
+		UObject* AnimationOrSkeleton = AssetData ? AssetData->GetOuter() : nullptr;
+		if (AnimationOrSkeleton)
+		{
+			SetAppliedModifier(AnimationOrSkeleton, DuplicateObject(this, GetOuter()));
+		}
+		Modify();
 	}
 }
 
@@ -269,10 +415,41 @@ void UAnimationModifier::PostLoad()
 			MarkPackageDirty();
 		}
 	}
-	// Non CDO, update revision GUID
-	else if(UAnimationModifier* TypedDefaultObject = Cast<UAnimationModifier>(DefaultObject))
+
+	if (PreviouslyAppliedModifier_DEPRECATED)
 	{
-		RevisionGuid = TypedDefaultObject->RevisionGuid;
+		// The applied guid is now stored at the revision guid of applied modifier
+		PreviouslyAppliedModifier_DEPRECATED->RevisionGuid = this->AppliedGuid_DEPRECATED;
+		if (UAnimationModifiersAssetUserData* AssetData = Cast<UAnimationModifiersAssetUserData>(GetOuter()))
+		{
+			if (UAnimSequence* AnimSequence = Cast<UAnimSequence>(AssetData->GetOuter())) 
+			{
+				UE_LOG(LogAnimation, Log, TEXT("Upgrading Applied Animation Modifier on Sequence %s."), *AssetData->GetOuter()->GetName());
+				UAnimationModifier* AppliedModifier = PreviouslyAppliedModifier_DEPRECATED;
+				PreviouslyAppliedModifier_DEPRECATED = nullptr;
+				SetAppliedModifier(AnimSequence, AppliedModifier);
+			}
+			else if (USkeleton* Skeleton = Cast<USkeleton>(AssetData->GetOuter()))
+			{
+				UE_LOG(LogAnimation, Warning, TEXT("Upgrading Applied Animation Modifier %s on Skeleton %s, please reapply the modifier."), *GetName(), *AssetData->GetOuter()->GetName());
+				UAnimationModifier* AppliedModifier = PreviouslyAppliedModifier_DEPRECATED;
+				PreviouslyAppliedModifier_DEPRECATED = nullptr;
+				SetAppliedModifier(Skeleton, AppliedModifier);
+			}
+			else
+			{
+				UE_LOG(LogAnimation, Error, TEXT("Upgrading Applied Animation Modifier on Unknown type of asset %s, Ignored."), *AssetData->GetOuter()->GetName());
+				PreviouslyAppliedModifier_DEPRECATED->MarkAsGarbage();
+				PreviouslyAppliedModifier_DEPRECATED = nullptr;
+			}
+		}
+		else 
+		{
+			UE_LOG(LogAnimation, Error, TEXT("Upgrading Applied Animation Modifier on Unknown asset, Ignored."));
+			PreviouslyAppliedModifier_DEPRECATED->MarkAsGarbage();
+			PreviouslyAppliedModifier_DEPRECATED = nullptr;
+		}
+		this->Modify();
 	}
 }
 
@@ -287,14 +464,6 @@ void UAnimationModifier::UpdateRevisionGuid(UClass* ModifierClass)
 	{
 		RevisionGuid = FGuid::NewGuid();
 
-		// Apply to any currently loaded instances of this class
-		for (TObjectIterator<UAnimationModifier> It; It; ++It)
-		{
-			if (*It != this && It->GetClass() == ModifierClass)
-			{
-				It->SetInstanceRevisionGuid(RevisionGuid);
-			}
-		}
 	}
 }
 
@@ -320,6 +489,24 @@ void UAnimationModifier::UpdateNativeRevisionGuid()
 	}
 }
 
+void UAnimationModifier::ExecuteOnRevert(UAnimSequence* InAnimSequence)
+{
+	CurrentAnimSequence = InAnimSequence;
+	CurrentSkeleton = InAnimSequence->GetSkeleton();
+	OnRevert(InAnimSequence);
+	CurrentAnimSequence = nullptr;
+	CurrentSkeleton = nullptr;
+}
+
+void UAnimationModifier::ExecuteOnApply(UAnimSequence* InAnimSequence)
+{	
+	CurrentAnimSequence = InAnimSequence;
+	CurrentSkeleton = InAnimSequence->GetSkeleton();
+	OnApply(InAnimSequence);
+	CurrentAnimSequence = nullptr;
+	CurrentSkeleton = nullptr;
+}
+
 void UAnimationModifier::ApplyToAll(TSubclassOf<UAnimationModifier> ModifierSubClass, bool bForceApply /*= true*/)
 {
 	if (UClass* ModifierClass = ModifierSubClass.Get())
@@ -328,31 +515,34 @@ void UAnimationModifier::ApplyToAll(TSubclassOf<UAnimationModifier> ModifierSubC
 		LoadModifierReferencers(ModifierSubClass);
 		
 		const FScopedTransaction Transaction(LOCTEXT("UndoAction_ApplyModifiers", "Applying Animation Modifier to Animation Sequence(s)"));		
-		for (TObjectIterator<UAnimationModifier> It; It; ++It)
+		for (UAnimationModifiersAssetUserData* AssetUserData : TObjectRange<UAnimationModifiersAssetUserData>{})
 		{
-			// Check if valid, of the required class, not pending kill and not a modifier back-up for reverting
-			const bool bIsRevertModifierInstance = *It && It->GetFName().ToString().StartsWith(RevertModifierObjectName.ToString());
-			if (*It && It->GetClass() == ModifierClass && IsValidChecked(*It) && !bIsRevertModifierInstance)
+			UAnimSequence* AnimSequence = Cast<UAnimSequence>(AssetUserData->GetOuter());
+			USkeleton* Skeleton = Cast<USkeleton>(AssetUserData->GetOuter());
+			ensure(AnimSequence || Skeleton);
+
+			for (UAnimationModifier* Modifier : AssetUserData->AnimationModifierInstances)
 			{
-				if (bForceApply || !It->IsLatestRevisionApplied())
+				if (Modifier->GetClass() == ModifierClass && IsValidChecked(Modifier))
 				{
-					// Go through outer chain to find AnimSequence
-					UObject* Outer = It->GetOuter();
-
-					while(Outer && !Outer->IsA<UAnimSequence>())
+					if (AnimSequence)
 					{
-						Outer = Outer->GetOuter();
+						if (bForceApply || !Modifier->IsLatestRevisionApplied(AnimSequence))
+						{
+							Modifier->ApplyToAnimationSequence(AnimSequence);
+						}
 					}
-
-					if (UAnimSequence* AnimSequence = Cast<UAnimSequence>(Outer))
+					else if (Skeleton) 
 					{
-						AnimSequence->Modify();
-						It->ApplyToAnimationSequence(AnimSequence);
-					}	
+						// TODO : Modifier on Skeleton
+						// We can apply the modifier to all animation sequences referencing this skeleton
+						// Save this behavior change for another CL
+						UE_LOG(LogAnimation, Warning, TEXT("Animation Modifier %s on Skeleton %s was skipped, please manually apply it from Skeleton"), *ModifierClass->GetName(), *Skeleton->GetName());
+					}
 				}
 			}
 		}
-	}	
+	}
 }
 
 void UAnimationModifier::LoadModifierReferencers(TSubclassOf<UAnimationModifier> ModifierSubClass)
@@ -380,19 +570,5 @@ const UAnimSequence* UAnimationModifier::GetAnimationSequence()
 	return CurrentAnimSequence;
 }
 
-void UAnimationModifier::UpdateStoredRevisions()
-{
-	AppliedGuid = RevisionGuid;
-}
-
-void UAnimationModifier::ResetStoredRevisions()
-{
-	AppliedGuid.Invalidate();
-}
-
-void UAnimationModifier::SetInstanceRevisionGuid(FGuid Guid)
-{
-	RevisionGuid = Guid;
-}
 
 #undef LOCTEXT_NAMESPACE
diff --git a/Engine/Source/Editor/AnimationModifiers/Private/AnimationModifierHelpers.h b/Engine/Source/Editor/AnimationModifiers/Private/AnimationModifierHelpers.h
--- a/Engine/Source/Editor/AnimationModifiers/Private/AnimationModifierHelpers.h
+++ b/Engine/Source/Editor/AnimationModifiers/Private/AnimationModifierHelpers.h
@@ -11,6 +11,7 @@
 #include "ClassViewerFilter.h"
 #include "AnimationModifier.h"
 #include "AnimationModifiersAssetUserData.h"
+#include "AnimationModifier.h"
 
 class FAnimationModifierHelpers
 {
@@ -75,4 +76,42 @@ public:
 		return AssetUserData;
 	}
 
+	//! @brief Enumerate the List[ModifierName,RevisionGuid] formed tag
+	//! @param[in] Tag The tag to parse
+	//! @param[in] Callback (FStringView, FGuid) -> bool, the function to apply to each modifier=revision pair, return false to _break_ from enumeration
+	template<typename Func>
+	static void EnumerateAnimationModifierTags(FStringView Tag, Func&& Callback)
+	{
+		static_assert(std::is_invocable_r_v<bool, Func, FStringView, FGuid>, "Func must be invocable (FStringView, FGuid) -> bool");
+		int Begin = 0;
+		int Mid = -1;
+		int End = -1;
+		for (size_t I = 0; I <= Tag.Len(); I++)
+		{
+			if (I == Tag.Len() || Tag[I] == UAnimationModifier::AnimationModifiersDelimiter)
+			{
+				if (I <= Begin)
+				{
+					continue; // Skip empty pair
+				}
+				End = I;
+				ensure(Mid >= Begin && End > Mid);
+				FStringView Name = Tag.SubStr(Begin, Mid - Begin);
+				FStringView Revision = Tag.SubStr(Mid + 1, End - Mid - 1);
+				FGuid RevisionGuid;
+				FGuid::Parse(FString{Revision}, RevisionGuid);
+
+				bool Break = !Invoke(Callback, Name, RevisionGuid);
+				if (Break)
+				{
+					break;
+				}
+				Begin = End + 1;
+			}
+			else if (Tag[I] == UAnimationModifier::AnimationModifiersAssignment)
+			{
+				Mid = I;
+			}
+		}
+	}
 };
diff --git a/Engine/Source/Editor/AnimationModifiers/Private/AnimationModifiersAssetUserData.cpp b/Engine/Source/Editor/AnimationModifiers/Private/AnimationModifiersAssetUserData.cpp
--- a/Engine/Source/Editor/AnimationModifiers/Private/AnimationModifiersAssetUserData.cpp
+++ b/Engine/Source/Editor/AnimationModifiers/Private/AnimationModifiersAssetUserData.cpp
@@ -19,6 +19,7 @@ void UAnimationModifiersAssetUserData::RemoveAnimationModifierInstance(UAnimatio
 {
 	checkf(AnimationModifierInstances.Contains(Instance), TEXT("Instance suppose to be removed is not found"));
 	AnimationModifierInstances.Remove(Instance);
+	AppliedModifiers.Remove(Instance);
 }
 
 const TArray<UAnimationModifier*>& UAnimationModifiersAssetUserData::GetAnimationModifierInstances() const
@@ -87,4 +88,35 @@ void UAnimationModifiersAssetUserData::RemoveInvalidModifiers()
 {
 	// This will catch force-deleted blueprints to be removed from our stored array
 	AnimationModifierInstances.RemoveAll([](UAnimationModifier* Modifier) { return Modifier == nullptr; });
+
+	// The AssetUserData on this animation's skeleton
+	UAnimationModifiersAssetUserData* SkeletonAssetUserData = nullptr;
+	if (const UObject* OwnerAsset = GetOuter())
+	{
+		if (const UAnimSequence* Sequence = Cast<UAnimSequence>(OwnerAsset->GetOuter()))
+		{
+			if (USkeleton* Skeleton = Sequence->GetSkeleton())
+			{
+				SkeletonAssetUserData = Skeleton->GetAssetUserData< UAnimationModifiersAssetUserData>();
+			}
+		}
+	}
+
+	auto IsAppliedModifierValid = [this, SkeletonAssetUserData](const decltype(AppliedModifiers)::ElementType& Pair)
+	{
+		// The key (template modifier) from both animation sequence or skeleton should be loaded before this object
+		if (const UObject* MasterModifier = Pair.Key.ResolveObject())
+		{
+			if (const UObject* AssetData = MasterModifier->GetOuter())
+			{
+				// If the MasterModifier is not owned by us or the skeleton
+				// This asset user data must been re-parented
+				// Remove those stale applied-modifier information
+				return AssetData == this || AssetData == SkeletonAssetUserData;
+			}
+		}
+		return false;
+	};
+	// Remove all applied modifier instances of deleted modifier
+	AppliedModifiers = AppliedModifiers.FilterByPredicate([](const decltype(AppliedModifiers)::ElementType& Pair) { return Pair.Key.ResolveObject() != nullptr; });
 }
diff --git a/Engine/Source/Editor/AnimationModifiers/Private/AnimationModifiersModule.cpp b/Engine/Source/Editor/AnimationModifiers/Private/AnimationModifiersModule.cpp
--- a/Engine/Source/Editor/AnimationModifiers/Private/AnimationModifiersModule.cpp
+++ b/Engine/Source/Editor/AnimationModifiers/Private/AnimationModifiersModule.cpp
@@ -59,6 +59,9 @@ void FAnimationModifiersModule::StartupModule()
 			GEditor->GetEditorSubsystem<UImportSubsystem>()->OnAssetReimport.AddRaw(this, &FAnimationModifiersModule::OnAssetPostReimport);
 		}
 	});
+
+	// Register extra asset registry tags for UAnimSequence
+	OnGetExtraObjectTagsHandle = UObject::FAssetRegistryTag::OnGetExtraObjectTags.AddStatic(&UAnimationModifier::GetAssetRegistryTagsForAppliedModifiersFromSkeleton);
 }
 
 TSharedRef<FApplicationMode> FAnimationModifiersModule::ExtendApplicationMode(const FName ModeName, TSharedRef<FApplicationMode> InMode)
@@ -142,6 +145,7 @@ void FAnimationModifiersModule::ShutdownModule()
 	if (GEditor)
 	{
 		GEditor->GetEditorSubsystem<UImportSubsystem>()->OnAssetPostImport.RemoveAll(this);
+		GEditor->GetEditorSubsystem<UImportSubsystem>()->OnAssetReimport.RemoveAll(this);
 	}
 }
 
@@ -170,6 +174,7 @@ void FAnimationModifiersModule::ShowAddAnimationModifierWindow(const TArray<UAni
 	}
 
 	FSlateApplication::Get().AddModalWindow(Window, ParentWindow, false);
+	UObject::FAssetRegistryTag::OnGetExtraObjectTags.Remove(OnGetExtraObjectTagsHandle);
 }
 
 void FAnimationModifiersModule::ApplyAnimationModifiers(const TArray<UAnimSequence*>& InSequences, bool bForceApply /*= true*/)
@@ -186,11 +191,10 @@ void FAnimationModifiersModule::ApplyAnimationModifiers(const TArray<UAnimSequen
 			UAnimationModifiersAssetUserData* UserData = AnimationSequence->GetAssetUserData<UAnimationModifiersAssetUserData>();
 			if (UserData)
 			{
-				AnimationSequence->Modify();
 				const TArray<UAnimationModifier*>& ModifierInstances = UserData->GetAnimationModifierInstances();
 				for (UAnimationModifier* Modifier : ModifierInstances)
 				{
-					if (bForceApply || !Modifier->IsLatestRevisionApplied())
+					if (bForceApply || !Modifier->IsLatestRevisionApplied(AnimationSequence))
 					{
 						Modifier->ApplyToAnimationSequence(AnimationSequence);
 					}
diff --git a/Engine/Source/Editor/AnimationModifiers/Private/AnimationModifiersModule.h b/Engine/Source/Editor/AnimationModifiers/Private/AnimationModifiersModule.h
--- a/Engine/Source/Editor/AnimationModifiers/Private/AnimationModifiersModule.h
+++ b/Engine/Source/Editor/AnimationModifiers/Private/AnimationModifiersModule.h
@@ -38,6 +38,8 @@ protected:
 
 	FWorkflowApplicationModeExtender Extender;
 
+	FDelegateHandle OnGetExtraObjectTagsHandle;
+
 	/** Callbacks used to add and apply default animation modifier classes */
 	void OnAssetPostImport(UFactory* ImportFactory, UObject* ImportedObject);
 	void OnAssetPostReimport(UObject* ReimportedObject);
diff --git a/Engine/Source/Editor/AnimationModifiers/Private/SAnimationModifiersTab.cpp b/Engine/Source/Editor/AnimationModifiers/Private/SAnimationModifiersTab.cpp
--- a/Engine/Source/Editor/AnimationModifiers/Private/SAnimationModifiersTab.cpp
+++ b/Engine/Source/Editor/AnimationModifiers/Private/SAnimationModifiersTab.cpp
@@ -287,6 +287,7 @@ void SAnimationModifiersTab::OnApplyModifier(const TArray<TWeakObjectPtr<UAnimat
 
 	FScopedTransaction Transaction(LOCTEXT("ApplyModifiersTransaction", "Applying Animation Modifier(s)"));
 	ApplyModifiers(ModifierInstances);
+	Refresh();
 }
 
 void SAnimationModifiersTab::FindAnimSequencesForSkeleton(TArray<UAnimSequence *>& ReferencedAnimSequences)
@@ -325,25 +326,68 @@ void SAnimationModifiersTab::OnRevertModifier(const TArray<TWeakObjectPtr<UAnima
 	
 	FScopedTransaction Transaction(LOCTEXT("RevertModifiersTransaction", "Reverting Animation Modifier(s)"));
 	RevertModifiers(ModifierInstances);
+	Refresh();
 }
 
 bool SAnimationModifiersTab::OnCanRevertModifier(const TArray<TWeakObjectPtr<UAnimationModifier>>& Instances)
 {
 	bool bCanRevert = false;
 
-	for (TWeakObjectPtr<UAnimationModifier> InstancePtr : Instances)
+	if (AnimationSequence)
 	{
-		checkf(InstancePtr.IsValid(), TEXT("Invalid weak object ptr to modifier instance"));
-		UAnimationModifier* Instance = InstancePtr.Get();
-
-		// At least one instance has to be revert-able
-		if (Instance->CanRevert())
+		for (const TWeakObjectPtr<UAnimationModifier>& InstancePtr : Instances)
 		{
-			bCanRevert = true;
-			break;
+			checkf(InstancePtr.IsValid(), TEXT("Invalid weak object ptr to modifier instance"));
+			const UAnimationModifier* Modifier = InstancePtr.Get();
+
+			// At least one instance has to be revert-able
+			if (Modifier->CanRevert(AnimationSequence))
+			{
+				bCanRevert = true;
+				break;
+			}
 		}
 	}
+	else if (Skeleton) 
+	{
+		// Check for modifiers applied on previous version
+		// Where asset registry tags were not written
+		for (const TWeakObjectPtr<UAnimationModifier>& Modifier : Instances)
+		{
+			if (Modifier->HasLegacyPreviousAppliedModifierOnSkeleton())
+			{
+				bCanRevert = true;
+				return true;
+			}
+		}
+
+		FAssetRegistryModule& AssetRegistryModule = FModuleManager::LoadModuleChecked<FAssetRegistryModule>(FName("AssetRegistry"));
+		IAssetRegistry& AssetRegistry = AssetRegistryModule.Get();
 
+		TSet<FString> Modifiers;
+		Algo::Transform(Instances, Modifiers, [](const TWeakObjectPtr<UAnimationModifier>& Ptr) 
+		{
+			return Ptr.Get()->GetName();
+		});
+
+		FARFilter Filter;
+		Filter.ClassPaths.Add(UAnimSequence::StaticClass()->GetClassPathName());
+		Filter.TagsAndValues.Add(FName{"Skeleton"}) = FAssetData(Skeleton).GetExportTextName();
+
+		AssetRegistry.EnumerateAssets(Filter, [&](const FAssetData& Asset) 
+		{
+			FString Tag = Asset.GetTagValueRef<FString>(UAnimationModifier::AnimationModifiersTag);
+
+			FAnimationModifierHelpers::EnumerateAnimationModifierTags(Tag, [&](FStringView Name, FGuid Revision) 
+			{
+				bCanRevert = Modifiers.ContainsByHash(GetTypeHash(Name), Name);
+				return !bCanRevert; // Break if bCanRevert is already true
+			});
+
+			return !bCanRevert; // Break if bCanRevert is already true
+		});
+
+	}
 	return bCanRevert;
 }
 
@@ -454,26 +498,26 @@ void SAnimationModifiersTab::ApplyModifiers(const TArray<UAnimationModifier*>& M
 		if (bApply)
 		{
 			FindAnimSequencesForSkeleton(AnimSequences);
-			Skeleton->Modify();
 		}
 	}
 
 	if (bApply)
 	{
 		UE::Anim::FApplyModifiersScope Scope;
-		for (UAnimSequence* AnimSequence : AnimSequences)
-		{
-			AnimSequence->Modify();
-		}
 
 		for (UAnimationModifier* Instance : Modifiers)
 		{
 			checkf(Instance, TEXT("Invalid modifier instance"));
-			Instance->Modify();
+			bool AppliedOnAny = false;
 			for (UAnimSequence* AnimSequence : AnimSequences)
 			{
 				ensure(!(Skeleton != nullptr) || AnimSequence->GetSkeleton() == Skeleton);
 				Instance->ApplyToAnimationSequence(AnimSequence);
+				AppliedOnAny = AppliedOnAny || Instance->CanRevert(AnimSequence);
+			}
+			if (AppliedOnAny && Skeleton)
+			{
+				Instance->RemoveLegacyPreviousAppliedModifierOnSkeleton(Skeleton);
 			}
 		}
 	}
@@ -496,26 +540,25 @@ void SAnimationModifiersTab::RevertModifiers(const TArray<UAnimationModifier*>&
 		if ( bRevert)
 		{
 			FindAnimSequencesForSkeleton(AnimSequences);
-			Skeleton->Modify();
 		}
 	}
 
 	if (bRevert)
 	{
-		for (UAnimSequence* AnimSequence : AnimSequences)
-		{
-			AnimSequence->Modify();
-		}
-
 		for (UAnimationModifier* Instance : Modifiers)
 		{
 			checkf(Instance, TEXT("Invalid modifier instance"));
-			Instance->Modify();
 			for (UAnimSequence* AnimSequence : AnimSequences)
 			{
 				ensure(!(Skeleton != nullptr) || AnimSequence->GetSkeleton() == Skeleton);
 				Instance->RevertFromAnimationSequence(AnimSequence);
 			}
+
+			if (Skeleton)
+			{
+				// Revert can not fail, thus we can always mark reverted
+				Instance->RemoveLegacyPreviousAppliedModifierOnSkeleton(Skeleton);
+			}
 		}
 	}	
 }
@@ -536,6 +579,7 @@ void SAnimationModifiersTab::RetrieveModifierData()
 			Item->Class = Modifier->GetClass();
 			Item->Index = ModifierIndex;
 			Item->OuterClass = AssetUserData->GetOuter()->GetClass();
+			Item->OutOfDate = AnimationSequence && !Modifier->IsLatestRevisionApplied(AnimationSequence);
 			ModifierItems.Add(ModifierListviewItem(Item));
 
 			// Register a delegate for when a BP is compiled, this so we can refresh the UI and prevent issues with invalid instance data
@@ -547,6 +591,82 @@ void SAnimationModifiersTab::RetrieveModifierData()
 				DelegateRegisteredBlueprints.Add(Blueprint);
 			}
 		}
+
+		// Refresh OutOfDate data for modifiers on Skeleton
+		if (Skeleton)
+		{
+			TArray<FModifierListviewItem*> ModifiersToCheck; // Modifiers that are up to date (until current search progress)
+			TArray<FModifierListviewItem*> NextModifiersToCheck; // Double buffer for ModifiersToCheck, store the UpToDate modifiers on this animation
+			ModifiersToCheck.Reserve(ModifierItems.Num());
+
+			// Handle backward compatibility 
+			// Modifier applied in previous version will not have the applied version stored in asset registry tags
+			// Exclude them in asset registry search and use the legacy OutOfDate check
+			for (TSharedPtr<FModifierListviewItem>& Item : ModifierItems)
+			{
+				ensure(!Item->OutOfDate);
+				if (Item->Instance->HasLegacyPreviousAppliedModifierOnSkeleton())
+				{
+					Item->OutOfDate = !Item->Instance->IsLatestRevisionApplied(Skeleton);
+				}
+				else
+				{
+					ModifiersToCheck.Add(Item.Get());
+				}
+			}
+
+			if (ModifiersToCheck.IsEmpty())
+			{
+				return;
+			}
+
+			NextModifiersToCheck.Reserve(ModifiersToCheck.Num());
+
+			// Modifier.OutOfDate = Any(Animation in Skeleton.Animations, Animation=>Modifier.OutOfDateOn(Animation))
+			FAssetRegistryModule& AssetRegistryModule = FModuleManager::LoadModuleChecked<FAssetRegistryModule>(FName("AssetRegistry"));
+			IAssetRegistry& AssetRegistry = AssetRegistryModule.Get();
+
+			FARFilter Filter;
+			Filter.ClassPaths.Add(UAnimSequence::StaticClass()->GetClassPathName());
+			Filter.TagsAndValues.Add(FName{"Skeleton"}) = FAssetData(Skeleton).GetExportTextName();
+
+			AssetRegistry.EnumerateAssets(Filter, [&](const FAssetData& Asset) {
+				FString Tag = Asset.GetTagValueRef<FString>(UAnimationModifier::AnimationModifiersTag);
+
+				// First mark all modifiers to check as out of date
+				// Since modifiers omitted from the tag were not applied and thus out of date
+				for (FModifierListviewItem* Item : ModifiersToCheck)
+				{
+					Item->OutOfDate = true;
+				}
+
+				// Enumerate the current asset's tag 
+				// To check the revision of each applied modifier
+				// UpToDate modifiers will be add to NextModifiersToCheck to continue the search in next Asset
+				NextModifiersToCheck.Empty();
+				FAnimationModifierHelpers::EnumerateAnimationModifierTags(Tag, [&](FStringView Name, FGuid Revision) {
+					FName ModifierName {Name};
+					auto CompareModifierListViewItemName = [ModifierName](const FModifierListviewItem* Item) {
+						return Item->Instance->GetFName() == ModifierName;
+					};
+					if (FModifierListviewItem** PtrItem = ModifiersToCheck.FindByPredicate(CompareModifierListViewItemName))
+					{
+						FModifierListviewItem* Item = *PtrItem;
+						Item->OutOfDate = Item->Instance->GetLatestRevisionGuid() != Revision;
+						if (!Item->OutOfDate)
+						{
+							NextModifiersToCheck.Add(Item); // Only need to continue check modifiers that is update to date
+						}
+					}
+					return true;
+				});
+				// Swap the double buffer
+				Swap(NextModifiersToCheck, NextModifiersToCheck);
+
+				// Break if all modifiers are out of date
+				return !ModifiersToCheck.IsEmpty();
+			});
+		}
 	}
 }
 
diff --git a/Engine/Source/Editor/AnimationModifiers/Private/SModifierItemRow.cpp b/Engine/Source/Editor/AnimationModifiers/Private/SModifierItemRow.cpp
--- a/Engine/Source/Editor/AnimationModifiers/Private/SModifierItemRow.cpp
+++ b/Engine/Source/Editor/AnimationModifiers/Private/SModifierItemRow.cpp
@@ -63,7 +63,7 @@ FText SModifierItemRow::GetInstanceText() const
 	// Ensure we remove the modifier class postfix
 	LabelString.RemoveFromEnd(Postfix);
 
-	if (InternalItem->Instance.IsValid() && !InternalItem->Instance->IsLatestRevisionApplied())
+	if (InternalItem->OutOfDate)
 	{
 		LabelString.Append(" (Out of Date)");
 	}
diff --git a/Engine/Source/Editor/AnimationModifiers/Private/SModifierListview.h b/Engine/Source/Editor/AnimationModifiers/Private/SModifierListview.h
--- a/Engine/Source/Editor/AnimationModifiers/Private/SModifierListview.h
+++ b/Engine/Source/Editor/AnimationModifiers/Private/SModifierListview.h
@@ -39,6 +39,7 @@ struct FModifierListviewItem
 	TWeakObjectPtr<UAnimationModifier> Instance;
 	int32 Index;
 	UClass* OuterClass;
+	bool OutOfDate;
 };
 
 typedef TSharedPtr<FModifierListviewItem> ModifierListviewItem;
diff --git a/Engine/Source/Editor/AnimationModifiers/Public/AnimationModifier.h b/Engine/Source/Editor/AnimationModifiers/Public/AnimationModifier.h
--- a/Engine/Source/Editor/AnimationModifiers/Public/AnimationModifier.h
+++ b/Engine/Source/Editor/AnimationModifiers/Public/AnimationModifier.h
@@ -38,8 +38,8 @@ public:
 	UAnimationModifier();
 
 	/** Applying and reverting the modifier for the given Animation Sequence */
-	void ApplyToAnimationSequence(class UAnimSequence* InAnimationSequence);
-	void RevertFromAnimationSequence(class UAnimSequence* InAnimationSequence);
+	void ApplyToAnimationSequence(UAnimSequence* AnimSequence) const;
+	void RevertFromAnimationSequence(UAnimSequence* AnimSequence) const;
 
 	/** Executed when the Animation is initialized (native event for debugging / testing purposes) */
 	UFUNCTION(BlueprintNativeEvent)
@@ -50,11 +50,27 @@ public:
 	void OnRevert(UAnimSequence* AnimationSequence);
 	virtual void OnRevert_Implementation(UAnimSequence* AnimationSequence) {}
 
-	/** Returns whether or not this modifier can be reverted, which means it will have to been applied (PreviouslyAppliedModifier != nullptr) */
-	bool CanRevert() const { return PreviouslyAppliedModifier != nullptr; };
+	/** Returns whether or not this modifier can be reverted, which means it will have to been applied to the given AnimationSequence */
+	bool CanRevert(IInterface_AssetUserData* AssetUserDataInterface) const;
 
 	/** Whether or not the latest compiled version of the blueprint is applied for this instance */
-	bool IsLatestRevisionApplied() const;
+	bool IsLatestRevisionApplied(IInterface_AssetUserData* AssetUserDataInterface) const;
+
+	//! @brief Get the latest revision GUID of this modifier class
+	FGuid GetLatestRevisionGuid() const;
+
+	//! @brief Get Asset Registry tags for applied modifiers from _skeleton_
+	//! @param[out] OutTags AnimationModifiers = %PATH%=%REVISION%;...
+	static void GetAssetRegistryTagsForAppliedModifiersFromSkeleton(const UObject* Object, TArray<UObject::FAssetRegistryTag>& OutTags);
+
+	static const FName AnimationModifiersTag;
+	static constexpr TCHAR AnimationModifiersDelimiter = ';';
+	static constexpr TCHAR AnimationModifiersAssignment = '=';
+
+	//! @brief Check if this Modifier On Skeleton has PreviouslyAppliedModifier_DEPRECATED from previous version
+	bool HasLegacyPreviousAppliedModifierOnSkeleton() const { return GetLegacyPreviouslyAppliedModifierForModifierOnSkeleton() != nullptr; }
+	//! @brief Mark this modifier on skeleton as reverted (affect CanRevert, IsLatestRevisionApplied)
+	void RemoveLegacyPreviousAppliedModifierOnSkeleton(USkeleton* Skeleton);
 
 	// Begin UObject Overrides
 	virtual void PostInitProperties() override;
@@ -84,24 +100,40 @@ protected:
 	void UpdateRevisionGuid(UClass* ModifierClass);
 	void UpdateNativeRevisionGuid();
 
+	void ExecuteOnRevert(UAnimSequence* InAnimSequence);
+	void ExecuteOnApply(UAnimSequence* InAnimSequence);
+	
 	/** Applies all instances of the provided Modifier class to its outer Animation Sequence*/
 	static void ApplyToAll(TSubclassOf<UAnimationModifier> ModifierSubClass, bool bForceApply = true);
 	static void LoadModifierReferencers(TSubclassOf<UAnimationModifier> ModifierSubClass);
 private:
-	UAnimSequence* CurrentAnimSequence;
-	USkeleton* CurrentSkeleton;
-
-	void UpdateStoredRevisions();
-	void ResetStoredRevisions();
-	void SetInstanceRevisionGuid(FGuid Guid);
-
-	// This holds the GUID representing the latest version of the modifier
+	// These value were not set during OnRevert, prefer to use the input from OnApply() or OnRevert()
+	UAnimSequence* CurrentAnimSequence = nullptr;
+	USkeleton* CurrentSkeleton = nullptr;
+
+	//! @brief Try get the applied modifier instance (associated with this modifier) on the animation sequence, null if not applied on the sequence
+	UAnimationModifier* GetAppliedModifier(IInterface_AssetUserData* AssetUserDataInterface) const;
+
+	//! @brief Set the applied modifier instance on the Animation Sequence
+	void SetAppliedModifier(TScriptInterface<IInterface_AssetUserData> AssetUserDataInterface, UAnimationModifier* AppliedModifier) const;
+	//! @brief Remove the applied modifier instance (associated with this modifier) from the Animation Sequence
+	//! @return The modifier instance removed
+	UAnimationModifier* FindAndRemoveAppliedModifier(TScriptInterface<IInterface_AssetUserData> AssetUserDataInterface) const;
+	//! @brief Get the applied modifier instance for modifier on skeleton read from _legacy_ version
+	//! @note This function is only intended for backward compatibility with existing asset
+	//!	@return PreviouslyAppliedModifier_DEPRECATED on Modifier on Skeleton, _before_ upgrade
+	UAnimationModifier* GetLegacyPreviouslyAppliedModifierForModifierOnSkeleton() const;
+
+
+	// - On class-default-objects, this is the latest revision GUID
+	// - On applied modifier instances, this is the revision applied
+	// - Not used for other instances
 	UPROPERTY(/*VisibleAnywhere for testing, Category = Revision*/)
 	FGuid RevisionGuid;
 
 	// This indicates whether or not the modifier is newer than what has been applied
 	UPROPERTY(/*VisibleAnywhere for testing, Category = Revision */)
-	FGuid AppliedGuid;
+	FGuid AppliedGuid_DEPRECATED;
 
 	// This holds the latest value returned by UpdateNativeRevisionGuid during the last PostLoad (changing this value will invalidate the GUIDs for all instances)
 	UPROPERTY(config)
@@ -109,9 +141,7 @@ private:
 
 	/** Serialized version of the modifier that has been previously applied to the Animation Asset */
 	UPROPERTY()
-	TObjectPtr<UAnimationModifier> PreviouslyAppliedModifier;
-
-	static const FName RevertModifierObjectName;
+	TObjectPtr<UAnimationModifier> PreviouslyAppliedModifier_DEPRECATED = nullptr;
 };
 
 namespace UE
diff --git a/Engine/Source/Editor/AnimationModifiers/Public/AnimationModifiersAssetUserData.h b/Engine/Source/Editor/AnimationModifiers/Public/AnimationModifiersAssetUserData.h
--- a/Engine/Source/Editor/AnimationModifiers/Public/AnimationModifiersAssetUserData.h
+++ b/Engine/Source/Editor/AnimationModifiers/Public/AnimationModifiersAssetUserData.h
@@ -25,6 +25,7 @@ class ANIMATIONMODIFIERS_API UAnimationModifiersAssetUserData : public UAssetUse
 	friend class SAnimationModifiersTab;
 	friend class SAnimationModifierContentBrowserWindow;
 	friend class FAnimationModifiersModule;
+	friend class UAnimationModifier;
 public:
 	const TArray<UAnimationModifier*>& GetAnimationModifierInstances() const;
 protected:	 
@@ -45,4 +46,28 @@ private:
 protected:
 	UPROPERTY()
 	TArray<TObjectPtr<UAnimationModifier>> AnimationModifierInstances;
+
+	// Animation modifiers APPLIED on the owning animation sequence
+	// 
+	// - Key = Modifier in AnimationModifierInstances of owning animation sequence or skeleton
+	// - Value = Modifier instance applied
+	// 
+	// Applied modifier instances stores (Properties, RevisionGuid, ...) at Apply Time
+	// Which are not visible or editable from user interface
+	// 
+	// In contrast, modifiers in AnimationModifierInstances 
+	// are objects displayed at the Animation Data Modifier window
+	// Where Properties are displayed for user editing anytime
+	// 
+	// Note, Modifier on Skeleton (MoS) applied instances are stored on each 
+	// animation sequence's asset user data, instead of the skeleton's
+	// this design is important to enable applied modifiers can be reverted for each sequence
+	// also ensure applying MoS (when [re]importing animation) will not affect the skeleton asset
+	//
+	// The only time MoS stores an applied instance here is upgrading from previous version
+	// The "UAnimationModifier::PreviouslyAppliedModifier_DEPRECATED" data will be migrated here
+	// To support proper revert for them
+	// For more compatibility handling, check UAnimationModifier::PostLoad() and GetAppliedModifier() 
+	UPROPERTY()
+	TMap<FSoftObjectPath, TObjectPtr<UAnimationModifier>> AppliedModifiers;
 };
