Fixup unacceptable string wording

diff --git a/Engine/Source/Editor/AnimationModifiers/Private/AnimationModifiersAssetUserData.cpp b/Engine/Source/Editor/AnimationModifiers/Private/AnimationModifiersAssetUserData.cpp
--- a/Engine/Source/Editor/AnimationModifiers/Private/AnimationModifiersAssetUserData.cpp
+++ b/Engine/Source/Editor/AnimationModifiers/Private/AnimationModifiersAssetUserData.cpp
@@ -105,11 +105,11 @@ void UAnimationModifiersAssetUserData::RemoveInvalidModifiers()
 	auto IsAppliedModifierValid = [this, SkeletonAssetUserData](const decltype(AppliedModifiers)::ElementType& Pair)
 	{
 		// The key (template modifier) from both animation sequence or skeleton should be loaded before this object
-		if (const UObject* MasterModifier = Pair.Key.ResolveObject())
+		if (const UObject* TemplateModifier = Pair.Key.ResolveObject())
 		{
-			if (const UObject* AssetData = MasterModifier->GetOuter())
+			if (const UObject* AssetData = TemplateModifier->GetOuter())
 			{
-				// If the MasterModifier is not owned by us or the skeleton
+				// If the TemplateModifier is not owned by us or the skeleton
 				// This asset user data must been re-parented
 				// Remove those stale applied-modifier information
 				return AssetData == this || AssetData == SkeletonAssetUserData;
