Fixup incorrect merge (see CL 23247896)

diff --git a/Engine/Source/Editor/AnimationModifiers/Private/AnimationModifiersModule.cpp b/Engine/Source/Editor/AnimationModifiers/Private/AnimationModifiersModule.cpp
--- a/Engine/Source/Editor/AnimationModifiers/Private/AnimationModifiersModule.cpp
+++ b/Engine/Source/Editor/AnimationModifiers/Private/AnimationModifiersModule.cpp
@@ -277,37 +277,6 @@ void FAnimationModifiersModule::OnAssetPostReimport(UObject* ReimportedObject)
 	}
 }
 
-void FAnimationModifiersModule::ShutdownModule()
-{
-	// Make sure we unregister the class layout 
-	FPropertyEditorModule* PropertyEditorModule = FModuleManager::GetModulePtr<FPropertyEditorModule>("PropertyEditor");
-	if (PropertyEditorModule)
-	{
-		PropertyEditorModule->UnregisterCustomClassLayout("AnimationModifier");
-	}
-
-	// Remove extender delegate
-	FWorkflowCentricApplication::GetModeExtenderList().RemoveAll([this](FWorkflowApplicationModeExtender& StoredExtender) { return StoredExtender.GetHandle() == Extender.GetHandle(); });
-
-	// During shutdown clean up all factories from any modes which are still active/alive
-	for (TWeakPtr<FApplicationMode> WeakMode : RegisteredApplicationModes)
-	{
-		if (WeakMode.IsValid())
-		{
-			TSharedPtr<FApplicationMode> Mode = WeakMode.Pin();
-			Mode->RemoveTabFactory(FAnimationModifiersTabSummoner::AnimationModifiersName);
-		}
-	}
-
-	RegisteredApplicationModes.Empty();
-
-	if (GEditor)
-	{
-		GEditor->GetEditorSubsystem<UImportSubsystem>()->OnAssetPostImport.RemoveAll(this);
-		GEditor->GetEditorSubsystem<UImportSubsystem>()->OnAssetReimport.RemoveAll(this);
-	}
-}
-
 void FAnimationModifiersModule::ShowAddAnimationModifierWindow(const TArray<UAnimSequence*>& InSequences)
 {
 	TSharedPtr<SAnimationModifierContentBrowserWindow> WindowContent;
