Animation Modifier : Improved Error Handling
- Dll Export the FApplyModifiersScope object to allow batch apply modifiers by script (Python or Blueprint Scripted Actions)
- New API to suppress warning or error dialogs when applying modifiers from script
- Fix : 'YesToAll' 'NoToAll' not respected when both error and warning are throwed from same modifier class

Contributed by [REDACTED_PERSON]

diff --git a/Engine/Source/Editor/AnimationModifiers/Private/AnimationModifier.cpp b/Engine/Source/Editor/AnimationModifiers/Private/AnimationModifier.cpp
--- a/Engine/Source/Editor/AnimationModifiers/Private/AnimationModifier.cpp
+++ b/Engine/Source/Editor/AnimationModifiers/Private/AnimationModifier.cpp
@@ -35,25 +35,8 @@
 
 #define LOCTEXT_NAMESPACE "AnimationModifier"
 
-int32 UE::Anim::FApplyModifiersScope::ScopesOpened = 0;
-TMap<FObjectKey, TOptional<EAppReturnType::Type>>  UE::Anim::FApplyModifiersScope::PerClassReturnTypeValues;
-
 const FName UAnimationModifier::AnimationModifiersTag = TEXT("AnimationModifierList");
 
-TOptional<EAppReturnType::Type> UE::Anim::FApplyModifiersScope::GetReturnType(const UAnimationModifier* InModifier)
-{
-	TOptional<EAppReturnType::Type>* ReturnTypePtr = PerClassReturnTypeValues.Find(FObjectKey(InModifier->GetClass()));
-	return ReturnTypePtr ? *ReturnTypePtr : TOptional<EAppReturnType::Type>();
-}
-
-void UE::Anim::FApplyModifiersScope::SetReturnType(const UAnimationModifier* InModifier, EAppReturnType::Type InReturnType)
-{
-	const FObjectKey Key(InModifier->GetClass());
-	ensure(!PerClassReturnTypeValues.Contains(Key));
-	PerClassReturnTypeValues.Add(Key, InReturnType);
-}
-
-
 UAnimationModifier::UAnimationModifier()
 {
 }
@@ -116,45 +99,24 @@ void UAnimationModifier::ApplyToAnimationSequence(class UAnimSequence* AnimSeque
 	bool bShouldRevert = bErrors;
 	
 	// If errors have occured - prompt user with OK and revert
-	TOptional<EAppReturnType::Type> ScopeReturnType = UE::Anim::FApplyModifiersScope::GetReturnType(this);	
-	static const FText MessageTitle = LOCTEXT("ModifierDialogTitle", "Modifier has generated errors during test run.");
+	static const FText MessageTitle = LOCTEXT("ModifierDialogTitle", "Modifier has generated errors.");
 	if (bErrors)
-	{		
-		if (UE::Anim::FApplyModifiersScope::ScopesOpened == 0 || !ScopeReturnType.IsSet() || ScopeReturnType.GetValue() != EAppReturnType::Type::Ok)
-		{
-			const FText ErrorMessageFormat = FText::FormatOrdered(LOCTEXT("ModifierErrorDescription", "Modifier: {0}\nAsset: {1}\n{2}\nResolve the errors before trying to apply again."), FText::FromString(GetClass()->GetPathName()),
-				FText::FromString(AnimSequence->GetPathName()), FText::FromString(OutputLog));
-			
-			EAppReturnType::Type ReturnValue = FMessageDialog::Open(EAppMsgType::Ok, ErrorMessageFormat, &MessageTitle);
-			UE::Anim::FApplyModifiersScope::SetReturnType(this, ReturnValue);
-		}
+	{
+		const FText ErrorMessage = FText::FormatOrdered(LOCTEXT("ModifierErrorDescription", "Modifier: {0}\nAsset: {1}\n{2}\nResolve the errors before trying to apply again."),
+			FText::FromString(GetClass()->GetPathName()),
+			FText::FromString(CurrentAnimSequence->GetPathName()),
+			FText::FromString(OutputLog));
+		UE::Anim::FApplyModifiersScope::HandleError(this, ErrorMessage, &MessageTitle);
 	}
 
 	// If _only_ warnings have occured - check if user has previously said YesAll / NoAll and process the result
 	if (bWarnings && !bShouldRevert)
 	{
-		if (UE::Anim::FApplyModifiersScope::ScopesOpened == 0 || !ScopeReturnType.IsSet())
-		{
-			const FText WarningMessage = FText::FormatOrdered(LOCTEXT("ModifierWarningDescription", "Modifier: {0}\nAsset: {1}\n{2}\nAre you sure you want to apply it?"), FText::FromString(GetClass()->GetPathName()),
-			FText::FromString(AnimSequence->GetPathName()), FText::FromString(OutputLog));
-
-			EAppReturnType::Type ReturnValue = FMessageDialog::Open(EAppMsgType::YesNoYesAllNoAll, WarningMessage, &MessageTitle);
-			bShouldRevert = ReturnValue == EAppReturnType::No || ReturnValue == EAppReturnType::NoAll;
-
-			// check if user response should be stored for further modifier applications
-			if(UE::Anim::FApplyModifiersScope::ScopesOpened > 0)
-			{
-				if (ReturnValue == EAppReturnType::Type::YesAll || ReturnValue == EAppReturnType::Type::NoAll)
-				{
-					UE::Anim::FApplyModifiersScope::SetReturnType(this, ReturnValue);
-				}
-			}
-		}
-		else
-		{
-			// Revert if previous user prompt return NoAll or if any errors occured previously 
-			bShouldRevert = ScopeReturnType.GetValue() == EAppReturnType::NoAll || ScopeReturnType.GetValue() == EAppReturnType::Ok;
-		}
+		const FText WarningMessage = FText::FormatOrdered(LOCTEXT("ModifierWarningDescription", "Modifier: {0}\nAsset: {1}\n{2}\nAre you sure you want to apply it?"),
+			FText::FromString(GetClass()->GetPathName()),
+			FText::FromString(CurrentAnimSequence->GetPathName()),
+			FText::FromString(OutputLog));
+		bShouldRevert = !UE::Anim::FApplyModifiersScope::HandleWarning(this, WarningMessage, &MessageTitle);
 	}
 
 	// Revert changes if necessary, otherwise post edit and refresh animation data
@@ -570,5 +532,103 @@ const UAnimSequence* UAnimationModifier::GetAnimationSequence()
 	return CurrentAnimSequence;
 }
 
+namespace UE::Anim
+{
+	TArray<FApplyModifiersScope::ESuppressionMode, TInlineAllocator<4>> FApplyModifiersScope::ScopeModeStack;
+	TMap<FObjectKey, EAppReturnType::Type> FApplyModifiersScope::WarningResponse;
+	TSet<FObjectKey> FApplyModifiersScope::ErrorResponse;
+
+	FApplyModifiersScope::ESuppressionMode FApplyModifiersScope::CurrentMode()
+	{
+		// Force Dialog is the legacy behavior when no scope is open
+		return ScopeModeStack.IsEmpty() ? ForceDialog : ScopeModeStack.Top();
+	}
+
+	FApplyModifiersScope::ESuppressionMode FApplyModifiersScope::Open(ESuppressionMode Mode)
+	{
+		if (Mode == NoChange)
+		{
+			// When no scope was open, we should be transit to ShowDialog mode (not ForceDialog mode)
+			Mode = ScopeModeStack.IsEmpty() ? ShowDialog : ScopeModeStack.Top();
+		}
+		if (ScopeModeStack.IsEmpty())
+		{
+			// Clear the response at the out most scope open
+			WarningResponse.Empty();
+			ErrorResponse.Empty();
+		}
+		ScopeModeStack.Push(Mode);
+		return Mode;
+	}
+
+	void FApplyModifiersScope::Close()
+	{
+		if (!ensureMsgf(!ScopeModeStack.IsEmpty(), TEXT("Unpaired FApplyModifiersScope::Open/Close call.")))
+		{
+			return;
+		}
+
+		ScopeModeStack.Pop();
+		if (ScopeModeStack.IsEmpty())
+		{
+			// Clear the response at the out most scope close
+			WarningResponse.Empty();
+			ErrorResponse.Empty();
+		}
+	}
+
+	void FApplyModifiersScope::HandleError(const UAnimationModifier* Modifier, const FText& Message, const FText* OptTitle)
+	{
+		ESuppressionMode Mode = CurrentMode();
+		if (Mode > SuppressWarningAndError && Mode < RevertAtWarning)
+		{
+			// Show the dialog if this error was not shown before, or when forced
+			if (Mode == ForceDialog || !ErrorResponse.Contains(Modifier->GetClass()))
+			{
+				FMessageDialog::Open(EAppMsgType::Ok, Message, OptTitle);
+			}
+		}
+		// Record the error to avoid same error dialog
+		ErrorResponse.Add(Modifier->GetClass());
+	}
+
+	bool FApplyModifiersScope::HandleWarning(const UAnimationModifier* Modifier, const FText& Message, const FText* OptTitle)
+	{
+		auto IsYesAllNoAll = [](EAppReturnType::Type Value) { return Value == EAppReturnType::YesAll || Value == EAppReturnType::NoAll; };
+		auto PtrValueOr = [](auto* Ptr, auto Default) { return Ptr ? *Ptr : Default; };
+
+		EAppReturnType::Type Response = EAppReturnType::No;
+		switch (ESuppressionMode Mode = CurrentMode())
+		{
+		case SuppressWarningAndError:
+		case SuppressWarning:
+			Response = EAppReturnType::Yes;
+			break;
+		case RevertAtWarning:
+			Response = EAppReturnType::No;
+			break;
+		case ForceDialog:
+			Response = EAppReturnType::Retry;
+			break;
+		case ShowDialog:
+		default:
+			Response = PtrValueOr(WarningResponse.Find(Modifier->GetClass()), EAppReturnType::Retry);
+			break;
+		}
+
+		// Don't show the dialog if the previous response is YesToAll or NoToAll
+		if (Response == EAppReturnType::Retry)
+		{
+			Response = FMessageDialog::Open(EAppMsgType::YesNoYesAllNoAll, Message, OptTitle);
+			// Record the YesAll NoAll response to avoid same the warning dialog
+			if (IsYesAllNoAll(Response))
+			{
+				WarningResponse.Add(Modifier->GetClass(), Response);
+			}
+		}
+
+		return (Response == EAppReturnType::Yes) || (Response == EAppReturnType::YesAll);
+	}
+} // namespace UE::Anim
 
 #undef LOCTEXT_NAMESPACE
diff --git a/Engine/Source/Editor/AnimationModifiers/Private/SAnimationModifierContentBrowserWindow.cpp b/Engine/Source/Editor/AnimationModifiers/Private/SAnimationModifierContentBrowserWindow.cpp
--- a/Engine/Source/Editor/AnimationModifiers/Private/SAnimationModifierContentBrowserWindow.cpp
+++ b/Engine/Source/Editor/AnimationModifiers/Private/SAnimationModifierContentBrowserWindow.cpp
@@ -233,6 +233,8 @@ FReply SAnimationModifierContentBrowserWindow::OnApply()
 		AssetUserData.Add(UserData);
 	}
 
+	UE::Anim::FApplyModifiersScope Scope;
+
 	// For each added modifier create add a new instance to each of the user data entries, using the one(s) set up in the window as template(s)
 	for (UAnimationModifier* Modifier : Modifiers)
 	{
diff --git a/Engine/Source/Editor/AnimationModifiers/Public/AnimationModifier.h b/Engine/Source/Editor/AnimationModifiers/Public/AnimationModifier.h
--- a/Engine/Source/Editor/AnimationModifiers/Public/AnimationModifier.h
+++ b/Engine/Source/Editor/AnimationModifiers/Public/AnimationModifier.h
@@ -148,32 +148,67 @@ namespace UE
 {
 	namespace Anim
 	{
-		struct FApplyModifiersScope
+		// RAII object to determine how Animation Modifier warning/errors are handled (ignore, user dialog, etc.)
+		struct ANIMATIONMODIFIERS_API FApplyModifiersScope
 		{
-			FApplyModifiersScope()
+			friend UAnimationModifier;
+
+			enum ESuppressionMode : uint8
+			{
+				// Do not change the error handling mode on this scope
+				// Use the mode set by parent scope or default (ShowDialog)
+				NoChange,
+				// Suppress error dialogs
+				// Suppress warnings dialogs, always apply modifiers
+				// No user interaction required
+				SuppressWarningAndError,
+				// Show error dialogs
+				// Suppress warnings dialogs, always apply modifiers
+				SuppressWarning,
+				// Show warning and error dialogs for first encounter
+				// Error dialog for each modifier class will only be showed once
+				ShowDialog,
+				// Always show the error or warning dialog
+				// Default behavior when no scope was open
+				ForceDialog,
+				// Suppress error dialogs
+				// Suppress warnings dialogs, always revert modifiers
+				// No user interaction required
+				RevertAtWarning,
+			};
+
+			FApplyModifiersScope(const FApplyModifiersScope&) = delete;
+			explicit FApplyModifiersScope(ESuppressionMode Mode = NoChange)
 			{
-				if (ScopesOpened == 0)
-				{				
-					PerClassReturnTypeValues.Empty();
-				}
-				++ScopesOpened;
+				Open(Mode);
 			}
 
 			~FApplyModifiersScope()
 			{
-				--ScopesOpened;
-				check(ScopesOpened >= 0);
-				if(ScopesOpened == 0)
-				{
-					PerClassReturnTypeValues.Empty();
-				}
+				Close();
 			}
 
-			static TOptional<EAppReturnType::Type> GetReturnType(const UAnimationModifier* InModifier);			
-			static void SetReturnType(const class UAnimationModifier* InModifier, EAppReturnType::Type InReturnType);
+protected:
+			/** Determine how to handle an Animation Modifier error, and execute accordingly */
+			static void HandleError(const UAnimationModifier* Modifier, const FText& Message, const FText* OptTitle = nullptr);
 
-			static TMap<FObjectKey, TOptional<EAppReturnType::Type>> PerClassReturnTypeValues;
-			static int32 ScopesOpened;
+			/** Determine how to handle an Animation Modifier warning, and execute accordingly. Returns whether or not the warning was handled (true) or warrants reverting the applied Animation Modifier (false) */
+			static bool HandleWarning(const UAnimationModifier* Modifier, const FText& Message, const FText* OptTitle = nullptr);
+
+private:
+			// Open a scope to control error handling when batch applying animation modifiers
+			static ESuppressionMode Open(ESuppressionMode Mode = NoChange);
+			// Close the most recent scope
+			static void Close();
+
+			static ESuppressionMode CurrentMode();
+
+			/** Errors already acknowledged */
+			static TSet<FObjectKey> ErrorResponse;
+			/** Warnings already ignored or treated as error */
+			static TMap<FObjectKey, EAppReturnType::Type> WarningResponse;
+			/** Error handle mode stack for scopes */
+			static TArray<ESuppressionMode, TInlineAllocator<4>> ScopeModeStack;
 		};
 	}
 }
