Fix infinite recursion due to bReapplyPostOwnerChange flag being set on an Anim Modifier instance

diff --git a/Engine/Source/Editor/AnimationModifiers/Private/AnimationModifiersAssetUserData.cpp b/Engine/Source/Editor/AnimationModifiers/Private/AnimationModifiersAssetUserData.cpp
--- a/Engine/Source/Editor/AnimationModifiers/Private/AnimationModifiersAssetUserData.cpp
+++ b/Engine/Source/Editor/AnimationModifiers/Private/AnimationModifiersAssetUserData.cpp
@@ -61,14 +61,18 @@ void UAnimationModifiersAssetUserData::PostEditChangeOwner()
 		return;
 	}
 
-	for (UAnimationModifier* Modifier : AnimationModifierInstances)
+	// Only react to PostEditChange if not currently applying modifiers
+	if(!UE::Anim::FApplyModifiersScope::IsScopePending())
 	{
-		if (!Modifier->bReapplyPostOwnerChange || Modifier->IsCurrentlyApplyingModifier())
-		{
-			continue;
-		}
-
-		Modifier->ApplyToAnimationSequence(ModifiedAnimSequence);
+	    UE::Anim::FApplyModifiersScope Scope;
+	    for (UAnimationModifier* Modifier : AnimationModifierInstances)
+	    {
+		    if (!Modifier->bReapplyPostOwnerChange)
+		    {
+			    continue;
+		    }
+		    Modifier->ApplyToAnimationSequence(ModifiedAnimSequence);
+	    }
 	}
 }
 #endif // WITH_EDITOR
diff --git a/Engine/Source/Editor/AnimationModifiers/Public/AnimationModifier.h b/Engine/Source/Editor/AnimationModifiers/Public/AnimationModifier.h
--- a/Engine/Source/Editor/AnimationModifiers/Public/AnimationModifier.h
+++ b/Engine/Source/Editor/AnimationModifiers/Public/AnimationModifier.h
@@ -83,6 +83,7 @@ public:
 	bool bReapplyPostOwnerChange = false;
 
 	/** Whether or not this modifier is in the process of being applied to an Animation Asset */
+	UE_DEPRECATED(5.2, "IsCurrentlyApplyingModifier has been deprecated, can check for modifiers being applied with UE::Anim::FApplyModifiersScope::IsScopePending")
 	bool IsCurrentlyApplyingModifier() const { return CurrentAnimSequence != nullptr || CurrentSkeleton != nullptr; };
 
 protected:
@@ -188,6 +189,7 @@ namespace UE
 				Close();
 			}
 
+			static bool IsScopePending() { return !ScopeModeStack.IsEmpty(); }
 protected:
 			/** Determine how to handle an Animation Modifier error, and execute accordingly */
 			static void HandleError(const UAnimationModifier* Modifier, const FText& Message, const FText* OptTitle = nullptr);
