Use the FMessageDialog overloads that pass the optional title by-value

#rnx

diff --git a/Engine/Plugins/Experimental/ChaosClothAssetEditor/Source/ChaosClothAssetTools/Private/ChaosClothAsset/ClothingAssetToClothAssetExporter.cpp b/Engine/Plugins/Experimental/ChaosClothAssetEditor/Source/ChaosClothAssetTools/Private/ChaosClothAsset/ClothingAssetToClothAssetExporter.cpp
--- a/Engine/Plugins/Experimental/ChaosClothAssetEditor/Source/ChaosClothAssetTools/Private/ChaosClothAsset/ClothingAssetToClothAssetExporter.cpp
+++ b/Engine/Plugins/Experimental/ChaosClothAssetEditor/Source/ChaosClothAssetTools/Private/ChaosClothAsset/ClothingAssetToClothAssetExporter.cpp
@@ -28,7 +28,7 @@ void UClothingAssetToChaosClothAssetExporter::Export(const UClothingAssetBase* C
 	{
 		const FText TitleMessage = LOCTEXT("ClothingAssetExporterTitle", "Error Exporting Clothing Asset");
 		const FText ErrorMessage = LOCTEXT("ClothingAssetExporterError", "Can only export from known ClothingAssetCommon types.");
-		FMessageDialog::Open(EAppMsgType::Ok, EAppReturnType::Ok, ErrorMessage, &TitleMessage);
+		FMessageDialog::Open(EAppMsgType::Ok, EAppReturnType::Ok, ErrorMessage, TitleMessage);
 		return;
 	}
 
diff --git a/Engine/Plugins/Experimental/RenderGrid/Source/RenderGridEditor/Private/UI/SRenderGridJobList.cpp b/Engine/Plugins/Experimental/RenderGrid/Source/RenderGridEditor/Private/UI/SRenderGridJobList.cpp
--- a/Engine/Plugins/Experimental/RenderGrid/Source/RenderGridEditor/Private/UI/SRenderGridJobList.cpp
+++ b/Engine/Plugins/Experimental/RenderGrid/Source/RenderGridEditor/Private/UI/SRenderGridJobList.cpp
@@ -533,11 +533,10 @@ TSharedRef<SWidget> UE::RenderGrid::Private::SRenderGridJobListTableRow::Generat
 							{
 								if (Grid->DoesJobIdExist(NewJobId))
 								{
-									const FText TitleText = LOCTEXT("JobIdNotUniqueTitle", "Duplicate Job IDs");
 									FMessageDialog::Open(
 										EAppMsgType::Ok,
 										FText::Format(LOCTEXT("JobIdNotUniqueMessage", "Job ID \"{0}\" is not unique."), FText::FromString(NewJobId)),
-										&TitleText);
+										LOCTEXT("JobIdNotUniqueTitle", "Duplicate Job IDs"));
 									return FText::FromString(OldJobId);
 								}
 
diff --git a/Engine/Plugins/MovieScene/MovieRenderPipeline/Source/MovieRenderPipelineEditor/Private/Widgets/SMoviePipelineConfigPanel.cpp b/Engine/Plugins/MovieScene/MovieRenderPipeline/Source/MovieRenderPipelineEditor/Private/Widgets/SMoviePipelineConfigPanel.cpp
--- a/Engine/Plugins/MovieScene/MovieRenderPipeline/Source/MovieRenderPipelineEditor/Private/Widgets/SMoviePipelineConfigPanel.cpp
+++ b/Engine/Plugins/MovieScene/MovieRenderPipeline/Source/MovieRenderPipelineEditor/Private/Widgets/SMoviePipelineConfigPanel.cpp
@@ -619,7 +619,7 @@ void SMoviePipelineConfigPanel::OnImportPreset(const FAssetData& InPresetAsset)
 		const FText TitleText = LOCTEXT("UnsavedPresetWarningTitle", "Unsaved Changes to Preset");
 		const FText MessageText = LOCTEXT("UnsavedPresetWarningMessage", "The changes made to the current preset will be lost by importing another preset. Do you want to continue with this import?");
 
-		if (FMessageDialog::Open(EAppMsgType::YesNo, MessageText, &TitleText) == EAppReturnType::No)
+		if (FMessageDialog::Open(EAppMsgType::YesNo, MessageText, TitleText) == EAppReturnType::No)
 		{
 			return;
 		}
diff --git a/Engine/Plugins/MovieScene/MovieRenderPipeline/Source/MovieRenderPipelineEditor/Private/Widgets/SMoviePipelineQueuePanel.cpp b/Engine/Plugins/MovieScene/MovieRenderPipeline/Source/MovieRenderPipelineEditor/Private/Widgets/SMoviePipelineQueuePanel.cpp
--- a/Engine/Plugins/MovieScene/MovieRenderPipeline/Source/MovieRenderPipelineEditor/Private/Widgets/SMoviePipelineQueuePanel.cpp
+++ b/Engine/Plugins/MovieScene/MovieRenderPipeline/Source/MovieRenderPipelineEditor/Private/Widgets/SMoviePipelineQueuePanel.cpp
@@ -755,7 +755,7 @@ void SMoviePipelineQueuePanel::OnImportSavedQueueAsset(const FAssetData& InPrese
 		const FText TitleText = LOCTEXT("UnsavedQueueWarningTitle", "Unsaved Changes to Queue");
 		const FText MessageText = LOCTEXT("UnsavedQueueWarningMessage", "The changes made to the current queue will be lost by importing another queue. Do you want to continue with this import?");
 
-		if (FMessageDialog::Open(EAppMsgType::YesNo, MessageText, &TitleText) == EAppReturnType::No)
+		if (FMessageDialog::Open(EAppMsgType::YesNo, MessageText, TitleText) == EAppReturnType::No)
 		{
 			return;
 		}
diff --git a/Engine/Plugins/Runtime/OpenXR/Source/OpenXRHMD/Private/OpenXRHMD.cpp b/Engine/Plugins/Runtime/OpenXR/Source/OpenXRHMD/Private/OpenXRHMD.cpp
--- a/Engine/Plugins/Runtime/OpenXR/Source/OpenXRHMD/Private/OpenXRHMD.cpp
+++ b/Engine/Plugins/Runtime/OpenXR/Source/OpenXRHMD/Private/OpenXRHMD.cpp
@@ -1682,7 +1682,7 @@ void ShowRestartWarning(const FText& Title)
 {
 	if (EAppReturnType::Ok == FMessageDialog::Open(EAppMsgType::OkCancel,
 		LOCTEXT("EditorRestartMsg", "The OpenXR runtime requires switching to a different GPU adapter, this requires an editor restart. Do you wish to restart now (you will be prompted to save any changes)?"),
-		&Title))
+		Title))
 	{
 		FUnrealEdMisc::Get().RestartEditor(false);
 	}
diff --git a/Engine/Source/Editor/AnimationModifiers/Private/AnimationModifier.cpp b/Engine/Source/Editor/AnimationModifiers/Private/AnimationModifier.cpp
--- a/Engine/Source/Editor/AnimationModifiers/Private/AnimationModifier.cpp
+++ b/Engine/Source/Editor/AnimationModifiers/Private/AnimationModifier.cpp
@@ -106,7 +106,7 @@ void UAnimationModifier::ApplyToAnimationSequence(class UAnimSequence* AnimSeque
 			FText::FromString(GetClass()->GetPathName()),
 			FText::FromString(CurrentAnimSequence->GetPathName()),
 			FText::FromString(OutputLog));
-		UE::Anim::FApplyModifiersScope::HandleError(this, ErrorMessage, &MessageTitle);
+		UE::Anim::FApplyModifiersScope::HandleError(this, ErrorMessage, MessageTitle);
 	}
 
 	// If _only_ warnings have occured - check if user has previously said YesAll / NoAll and process the result
@@ -116,7 +116,7 @@ void UAnimationModifier::ApplyToAnimationSequence(class UAnimSequence* AnimSeque
 			FText::FromString(GetClass()->GetPathName()),
 			FText::FromString(CurrentAnimSequence->GetPathName()),
 			FText::FromString(OutputLog));
-		bShouldRevert = !UE::Anim::FApplyModifiersScope::HandleWarning(this, WarningMessage, &MessageTitle);
+		bShouldRevert = !UE::Anim::FApplyModifiersScope::HandleWarning(this, WarningMessage, MessageTitle);
 	}
 
 	// Revert changes if necessary, otherwise post edit and refresh animation data
@@ -577,7 +577,7 @@ namespace UE::Anim
 		}
 	}
 
-	void FApplyModifiersScope::HandleError(const UAnimationModifier* Modifier, const FText& Message, const FText* OptTitle)
+	void FApplyModifiersScope::HandleError(const UAnimationModifier* Modifier, const FText& Message, const FText& Title)
 	{
 		ESuppressionMode Mode = CurrentMode();
 		if (Mode > SuppressWarningAndError && Mode < RevertAtWarning)
@@ -585,14 +585,14 @@ namespace UE::Anim
 			// Show the dialog if this error was not shown before, or when forced
 			if (Mode == ForceDialog || !ErrorResponse.Contains(Modifier->GetClass()))
 			{
-				FMessageDialog::Open(EAppMsgType::Ok, Message, OptTitle);
+				FMessageDialog::Open(EAppMsgType::Ok, Message, Title);
 			}
 		}
 		// Record the error to avoid same error dialog
 		ErrorResponse.Add(Modifier->GetClass());
 	}
 
-	bool FApplyModifiersScope::HandleWarning(const UAnimationModifier* Modifier, const FText& Message, const FText* OptTitle)
+	bool FApplyModifiersScope::HandleWarning(const UAnimationModifier* Modifier, const FText& Message, const FText& Title)
 	{
 		auto IsYesAllNoAll = [](EAppReturnType::Type Value) { return Value == EAppReturnType::YesAll || Value == EAppReturnType::NoAll; };
 		auto PtrValueOr = [](auto* Ptr, auto Default) { return Ptr ? *Ptr : Default; };
@@ -619,7 +619,7 @@ namespace UE::Anim
 		// Don't show the dialog if the previous response is YesToAll or NoToAll
 		if (Response == EAppReturnType::Retry)
 		{
-			Response = FMessageDialog::Open(EAppMsgType::YesNoYesAllNoAll, Message, OptTitle);
+			Response = FMessageDialog::Open(EAppMsgType::YesNoYesAllNoAll, Message, Title);
 			// Record the YesAll NoAll response to avoid same the warning dialog
 			if (IsYesAllNoAll(Response))
 			{
diff --git a/Engine/Source/Editor/AnimationModifiers/Public/AnimationModifier.h b/Engine/Source/Editor/AnimationModifiers/Public/AnimationModifier.h
--- a/Engine/Source/Editor/AnimationModifiers/Public/AnimationModifier.h
+++ b/Engine/Source/Editor/AnimationModifiers/Public/AnimationModifier.h
@@ -192,10 +192,10 @@ namespace UE
 			static bool IsScopePending() { return !ScopeModeStack.IsEmpty(); }
 protected:
 			/** Determine how to handle an Animation Modifier error, and execute accordingly */
-			static void HandleError(const UAnimationModifier* Modifier, const FText& Message, const FText* OptTitle = nullptr);
+			static void HandleError(const UAnimationModifier* Modifier, const FText& Message, const FText& Title);
 
 			/** Determine how to handle an Animation Modifier warning, and execute accordingly. Returns whether or not the warning was handled (true) or warrants reverting the applied Animation Modifier (false) */
-			static bool HandleWarning(const UAnimationModifier* Modifier, const FText& Message, const FText* OptTitle = nullptr);
+			static bool HandleWarning(const UAnimationModifier* Modifier, const FText& Message, const FText& Title);
 
 private:
 			// Open a scope to control error handling when batch applying animation modifiers
