Anim Modifiers - Can't add multiple modifiers to multiple animation sequences without the engine crashing
#fix prevent deadlocking when trying to cancel async animation compression task while holding onto evaluation/modification lock on current thread
#misc ensure that apply an animation modifier holds aforementioned lock
#misc add bracket to MotionExtractorModifier

#localization none
#tests tested cook on project that was timing out to confirm this fixed the issue

diff --git a/Engine/Source/Editor/AnimationModifiers/Private/AnimationModifier.cpp b/Engine/Source/Editor/AnimationModifiers/Private/AnimationModifier.cpp
--- a/Engine/Source/Editor/AnimationModifiers/Private/AnimationModifier.cpp
+++ b/Engine/Source/Editor/AnimationModifiers/Private/AnimationModifier.cpp
@@ -77,6 +77,7 @@ void UAnimationModifier::ApplyToAnimationSequence(class UAnimSequence* AnimSeque
 
 	{
 		// Group the OnRevert & OnApply operation into one Bracket to reduce compression request
+		IAnimationDataModel::FEvaluationAndModificationLock EvaluationAndModificationLock(*AnimSequence->GetDataModelInterface());
 		IAnimationDataController::FScopedBracket ScopedBracket(AnimSequence->GetController(), LOCTEXT("ApplyModifierBracket", "Applying Animation Modifier"));
 
 		/** In case this modifier has been previously applied, revert it using the serialised out version at the time */
