Applying only FootstepEffectTagModifier or FootFXAnimModifier and undoing the apply creates None notifies
#fix re-introduce marking AnimSequence as modified when applying animation-modifier(s), removed in CL 23191977

diff --git a/Engine/Source/Editor/AnimationModifiers/Private/AnimationModifier.cpp b/Engine/Source/Editor/AnimationModifiers/Private/AnimationModifier.cpp
--- a/Engine/Source/Editor/AnimationModifiers/Private/AnimationModifier.cpp
+++ b/Engine/Source/Editor/AnimationModifiers/Private/AnimationModifier.cpp
@@ -70,6 +70,8 @@ void UAnimationModifier::ApplyToAnimationSequence(class UAnimSequence* AnimSeque
 	UAnimationModifier* AppliedModifier = DuplicateObject(this, AppliedModifierOuter);
 	// Set the revision guid on applied modifier to latest
 	AppliedModifier->RevisionGuid = GetLatestRevisionGuid();
+	
+	AnimSequence->Modify();
 
 	FTransaction AnimationDataTransaction;
 	AnimationDataTransaction.SaveObject(AnimSequence);
