Removed experimental tag from UAnimationModifier

diff --git a/Engine/Source/Editor/AnimationModifiers/Public/AnimationModifier.h b/Engine/Source/Editor/AnimationModifiers/Public/AnimationModifier.h
--- a/Engine/Source/Editor/AnimationModifiers/Public/AnimationModifier.h
+++ b/Engine/Source/Editor/AnimationModifiers/Public/AnimationModifier.h
@@ -28,7 +28,7 @@ class USkeleton;
 struct FFrame;
 struct FObjectKey;
 
-UCLASS(Blueprintable, Experimental, config = Editor, defaultconfig)
+UCLASS(Blueprintable, config = Editor, defaultconfig)
 class ANIMATIONMODIFIERS_API UAnimationModifier : public UObject
 {
 	GENERATED_BODY()
