Animation Modifier improvements
- ensure that any _skeleton_ modifiers as OnPostOwnerChange also get applied when a UAnimSequence is edited

diff --git a/Engine/Source/Editor/AnimationModifiers/Private/AnimationModifiersAssetUserData.cpp b/Engine/Source/Editor/AnimationModifiers/Private/AnimationModifiersAssetUserData.cpp
--- a/Engine/Source/Editor/AnimationModifiers/Private/AnimationModifiersAssetUserData.cpp
+++ b/Engine/Source/Editor/AnimationModifiers/Private/AnimationModifiersAssetUserData.cpp
@@ -54,25 +54,39 @@ void UAnimationModifiersAssetUserData::PostEditChangeOwner()
 	{
 		return;
 	}
-
-	UAnimSequence* ModifiedAnimSequence = Cast<UAnimSequence>(GetOuter());
-	if (!ModifiedAnimSequence)
-	{
-		return;
-	}
-
+	
 	// Only react to PostEditChange if not currently applying modifiers
 	if(!UE::Anim::FApplyModifiersScope::IsScopePending())
-	{
-	    UE::Anim::FApplyModifiersScope Scope;
-	    for (UAnimationModifier* Modifier : AnimationModifierInstances)
-	    {
-		    if (!Modifier->bReapplyPostOwnerChange)
-		    {
-			    continue;
-		    }
-		    Modifier->ApplyToAnimationSequence(ModifiedAnimSequence);
-	    }
+	{		
+		if (UAnimSequence* ModifiedAnimSequence = Cast<UAnimSequence>(GetOuter()))
+		{
+			UE::Anim::FApplyModifiersScope Scope;
+			for (const UAnimationModifier* Modifier : AnimationModifierInstances)
+			{
+				if (!Modifier->bReapplyPostOwnerChange)
+				{
+					continue;
+				}
+				Modifier->ApplyToAnimationSequence(ModifiedAnimSequence);
+			}
+
+			// Find outer USkeleton, and apply any skeleton-level modifiers
+			if (USkeleton* OuterSkeleton = ModifiedAnimSequence->GetSkeleton())
+			{
+				if (UAnimationModifiersAssetUserData* SkeletonAssetUserData = OuterSkeleton->GetAssetUserData<UAnimationModifiersAssetUserData>())
+				{
+				    for (const UAnimationModifier* Modifier : SkeletonAssetUserData->GetAnimationModifierInstances())
+				    {
+					    if (!Modifier->bReapplyPostOwnerChange)
+					    {
+						    continue;
+					    }
+    
+					    Modifier->ApplyToAnimationSequence(ModifiedAnimSequence);
+				    }
+				}
+			}
+		}
 	}
 }
 #endif // WITH_EDITOR
