Add support for removing (and optionally also reverting) Animation Modifier(s) from a selection of Anim Sequences in the content browser.

- Added `SRemoveAnimationModifierContentBrowserWindow` widget.
- Added `ShowRemoveAnimationModifierWindow` virtual method to `IAnimationModifiersModule`.
- Added `OnSelectionChanged` delegate to SModifierListView.
- Fixed AnimationModifiersModule not properly unsubscribing its delegates bound to` UObject::FAssetRegistryTag::OnGetExtraObjectTags`.

[RN] minor

diff --git a/Engine/Source/Editor/AnimationModifiers/Private/AnimationModifiersModule.cpp b/Engine/Source/Editor/AnimationModifiers/Private/AnimationModifiersModule.cpp
--- a/Engine/Source/Editor/AnimationModifiers/Private/AnimationModifiersModule.cpp
+++ b/Engine/Source/Editor/AnimationModifiers/Private/AnimationModifiersModule.cpp
@@ -95,6 +95,8 @@ void FAnimationModifiersModule::ShutdownModule()
 	UToolMenus::UnregisterOwner(this);
 	FCoreDelegates::OnPostEngineInit.Remove(DelegateHandle);
 
+	UObject::FAssetRegistryTag::OnGetExtraObjectTags.Remove(OnGetExtraObjectTagsHandle);
+	
 	// Remove extender delegate
 	FWorkflowCentricApplication::GetModeExtenderList().RemoveAll([this](FWorkflowApplicationModeExtender& StoredExtender) { return StoredExtender.GetHandle() == Extender.GetHandle(); });
 
@@ -229,6 +231,18 @@ void FAnimationModifiersModule::RegisterMenus()
 						ApplyAnimationModifiers(AnimSequences, false);
 					}))
 				);
+
+				MenuBuilder.AddMenuEntry(
+					LOCTEXT("AnimSequence_RemoveAnimationModifier", "Remove Modifiers"),
+					LOCTEXT("AnimSequence_RemoveAnimationModifierTooltip", "Remove animation modifier(s)"),
+					FSlateIcon(FAppStyle::GetAppStyleSetName(), "ClassIcon.AnimationModifier"),
+					FUIAction(FExecuteAction::CreateLambda([GetAnimSequences, this]()
+					{
+						TArray<UAnimSequence*> AnimSequences;
+						GetAnimSequences(AnimSequences);
+
+						ShowRemoveAnimationModifierWindow(AnimSequences);
+					})));
 			});
 
 			IAssetTools& AssetTools = FModuleManager::LoadModuleChecked<FAssetToolsModule>("AssetTools").Get();
@@ -320,7 +334,33 @@ void FAnimationModifiersModule::ShowAddAnimationModifierWindow(const TArray<UAni
 	}
 
 	FSlateApplication::Get().AddModalWindow(Window, ParentWindow, false);
-	UObject::FAssetRegistryTag::OnGetExtraObjectTags.Remove(OnGetExtraObjectTagsHandle);
+}
+
+void FAnimationModifiersModule::ShowRemoveAnimationModifierWindow(const TArray<UAnimSequence*>& InSequences)
+{
+	TSharedPtr<SRemoveAnimationModifierContentBrowserWindow> WindowContent;
+
+	TSharedRef<SWindow> Window = SNew(SWindow)
+		.Title(LOCTEXT("WindowTitle", "Remove Animation Modifier(s)"))
+		.SizingRule(ESizingRule::UserSized)
+		.ClientSize(FVector2D(500, 500));
+
+	Window->SetContent
+	(
+		SAssignNew(WindowContent, SRemoveAnimationModifierContentBrowserWindow)
+		.WidgetWindow(Window)
+		.AnimSequences(InSequences)
+	);
+
+	TSharedPtr<SWindow> ParentWindow;
+
+	if (FModuleManager::Get().IsModuleLoaded("MainFrame"))
+	{
+		IMainFrameModule& MainFrame = FModuleManager::LoadModuleChecked<IMainFrameModule>("MainFrame");
+		ParentWindow = MainFrame.GetParentWindow();
+	}
+
+	FSlateApplication::Get().AddModalWindow(Window, ParentWindow, false);
 }
 
 void FAnimationModifiersModule::ApplyAnimationModifiers(const TArray<UAnimSequence*>& InSequences, bool bForceApply /*= true*/)
diff --git a/Engine/Source/Editor/AnimationModifiers/Private/AnimationModifiersModule.h b/Engine/Source/Editor/AnimationModifiers/Private/AnimationModifiersModule.h
--- a/Engine/Source/Editor/AnimationModifiers/Private/AnimationModifiersModule.h
+++ b/Engine/Source/Editor/AnimationModifiers/Private/AnimationModifiersModule.h
@@ -42,6 +42,7 @@ public:
 
 	/** Begin IAnimationModifiersModule overrides */
 	virtual void ShowAddAnimationModifierWindow(const TArray<UAnimSequence*>& InSequences) override;
+	virtual void ShowRemoveAnimationModifierWindow(const TArray<UAnimSequence*>& InSequences) override;
 	virtual void ApplyAnimationModifiers(const TArray<UAnimSequence*>& InSequences, bool bForceApply = true) override;
 	/** End IAnimationModifiersModule overrides */
 
diff --git a/Engine/Source/Editor/AnimationModifiers/Private/SAnimationModifierContentBrowserWindow.cpp b/Engine/Source/Editor/AnimationModifiers/Private/SAnimationModifierContentBrowserWindow.cpp
--- a/Engine/Source/Editor/AnimationModifiers/Private/SAnimationModifierContentBrowserWindow.cpp
+++ b/Engine/Source/Editor/AnimationModifiers/Private/SAnimationModifierContentBrowserWindow.cpp
@@ -29,6 +29,7 @@
 #include "ScopedTransaction.h"
 #include "SlotBase.h"
 #include "Styling/AppStyle.h"
+#include "Styling/SlateIconFinder.h"
 #include "Templates/SubclassOf.h"
 #include "Types/SlateEnums.h"
 #include "UObject/Class.h"
@@ -53,6 +54,10 @@ struct FGeometry;
 
 #define LOCTEXT_NAMESPACE "AnimationModifierContentBrowserWindow"
 
+////////////////////////////////////////////////////////////////////////////////////
+// Add Animation Modifier(s) widget
+////////////////////////////////////////////////////////////////////////////////////
+
 void SAnimationModifierContentBrowserWindow::Construct(const FArguments& InArgs)
 {
 	CreateInstanceDetailsView();
@@ -300,4 +305,392 @@ bool SAnimationModifierContentBrowserWindow::CanApply() const
 	return Modifiers.Num() > 0;
 }
 
-#undef LOCTEXT_NAMESPACE // "AnimationModifierContentBrowserWindow"
\ No newline at end of file
+////////////////////////////////////////////////////////////////////////////////////
+// Remove Animation Modifier(s) widget
+////////////////////////////////////////////////////////////////////////////////////
+
+/** Widget used to display a list of sequences that will be modified after an animation modifier application */
+class SAnimSequencesToBeModifiedViewer : public SCompoundWidget
+{
+	SLATE_BEGIN_ARGS(SAnimSequencesToBeModifiedViewer)
+	{
+	}
+
+	SLATE_END_ARGS()
+
+	void Construct(const FArguments& Args)
+	{
+		AnimSequencesToBeModifiedListView = SNew(SListView<UAnimSequence*>)
+				.ItemHeight(20.0f)
+				.ListItemsSource(&AnimSequencesToBeModifiedList)
+				.SelectionMode(ESelectionMode::None)
+				.OnGenerateRow(this, &SAnimSequencesToBeModifiedViewer::OnGenerateWidgetForRow);
+
+		ChildSlot
+		[
+			SNew(SHorizontalBox)
+			+ SHorizontalBox::Slot()
+			[
+				AnimSequencesToBeModifiedListView.ToSharedRef()
+			]
+		];
+	};
+
+	void UpdateToBeModifiedAnimSequencesList(const TArray<UAnimSequence*>& InAnimSequences, const UAnimationModifier * InModifier)
+	{
+		AnimSequencesToBeModifiedList.Empty();
+
+		if (InModifier != nullptr)
+		{
+			for (UAnimSequence* AnimationSequence : InAnimSequences)
+			{
+				if (const UAnimationModifiersAssetUserData* UserData = AnimationSequence->GetAssetUserData<UAnimationModifiersAssetUserData>())
+				{
+					const bool bIsModifierFoundInAnimSequence = UserData->GetAnimationModifierInstances().ContainsByPredicate(([InModifier](const UAnimationModifier* TestModifier)
+					{
+						return InModifier->GetClass() == TestModifier->GetClass();
+					}));
+				
+					if (bIsModifierFoundInAnimSequence)
+					{
+						AnimSequencesToBeModifiedList.Push(AnimationSequence);
+					}
+				}
+			}
+		}
+		
+		AnimSequencesToBeModifiedListView->RequestListRefresh();
+	}
+	
+private:
+
+	TSharedRef<ITableRow> OnGenerateWidgetForRow(UAnimSequence* InItem, const TSharedRef<STableViewBase>& OwnerTable) const
+	{
+		const FSlateBrush* ClassIcon = FSlateIconFinder::FindIconBrushForClass(UAnimSequence::StaticClass());
+		
+		return SNew(STableRow<UAnimSequence*>, OwnerTable)
+		[
+			SNew(SHorizontalBox)
+
+			// Logo
+			+SHorizontalBox::Slot()
+				.AutoWidth()
+				.Padding( 6.0f, 2.0f, 6.0f, 2.0f )
+				[
+					SNew( SImage )
+					.Image(	ClassIcon )
+					.Visibility( ClassIcon != FAppStyle::GetDefaultBrush()? EVisibility::Visible : EVisibility::Collapsed )
+				]
+
+			// Display name
+			+SHorizontalBox::Slot()
+				.FillWidth(1.0f)
+				.Padding( 0.0f, 3.0f, 6.0f, 3.0f )
+				.VAlign(VAlign_Center)
+				[
+					SNew( STextBlock )
+						.Text( FText::FromString(InItem->GetName()) )
+						.ColorAndOpacity(FSlateColor::UseForeground())
+				]
+		];
+	}
+
+	TSharedPtr<SListView<UAnimSequence*>> AnimSequencesToBeModifiedListView;
+	TArray<UAnimSequence*> AnimSequencesToBeModifiedList;
+};
+
+void SRemoveAnimationModifierContentBrowserWindow::Construct(const FArguments& InArgs)
+{
+	WidgetWindow = InArgs._WidgetWindow;
+	AnimSequences = InArgs._AnimSequences;
+	
+	this->ChildSlot
+	[
+		SNew(SOverlay)		
+		+SOverlay::Slot()
+		[				
+			SNew(SVerticalBox)
+			+ SVerticalBox::Slot()
+			.AutoHeight()
+			[
+				SNew(SBorder)
+				.Padding(2.0f)
+				.BorderImage(FAppStyle::GetBrush("ToolPanel.GroupBorder"))
+				[
+					SNew(SHorizontalBox)
+					+ SHorizontalBox::Slot()
+					.Padding(3.0f, 3.0f)
+					.AutoWidth()
+					[
+						SAssignNew(AddModifierCombobox, SComboButton)
+						.MenuContent()
+						[
+							SNew(SBox)
+							.WidthOverride(280)
+							[
+								SNew(SVerticalBox)
+								+ SVerticalBox::Slot()
+								.AutoHeight()
+								.MaxHeight(500)
+								[
+									GenerateAnimationModifierPicker()
+								]
+							]
+						]
+						.ButtonContent()
+						[
+							SNew(STextBlock)
+							.Text(LOCTEXT("AnimationModifierWindow_RemoveModifier", "Remove Modifier"))
+						]
+					]
+				]
+			]
+			
+			+ SVerticalBox::Slot()
+			[
+				SNew(SBorder)
+				.Padding(2.0f)
+				.BorderImage(FAppStyle::GetBrush("ToolPanel.GroupBorder"))
+				[				
+					SNew(SSplitter)
+					.Orientation(EOrientation::Orient_Vertical)
+
+					+ SSplitter::Slot()
+					.Value(.5f)
+					[
+						SNew(SBox)
+						.Padding(2.0f)
+						[
+							SAssignNew(ModifierListView, SModifierListView)
+							.Items(&ModifierItems)
+							.InstanceDetailsView(ModifierInstanceDetailsView)
+							.OnRemoveModifier(FOnModifierArray::CreateSP(this, &SRemoveAnimationModifierContentBrowserWindow::RemoveModifiersCallback))
+							.OnSelectedModifierChanged_Lambda([this](const TWeakObjectPtr<UAnimationModifier>& InSelectedModifier)
+							{
+								AnimSequencesToBeModifiedListViewer->UpdateToBeModifiedAnimSequencesList(AnimSequences, InSelectedModifier.Get());
+							})
+						]
+					]
+
+					// Show sequences that this modifier is affecting.
+					+ SSplitter::Slot()
+					.Value(.5f)
+					[
+						SNew(SVerticalBox)
+
+						+SVerticalBox::Slot()
+						.AutoHeight()
+						.Padding(6.f, 4.f)
+						[
+							SNew(STextBlock)
+							.Text(LOCTEXT("AnimationModifierWindow_AffectedTitle", "Sequence(s) affected by the selected Animation Modifier"))
+						]
+						
+						+SVerticalBox::Slot()
+						.Padding(2.0f)
+						[
+							SAssignNew(AnimSequencesToBeModifiedListViewer, SAnimSequencesToBeModifiedViewer)
+						]
+					]
+				]
+			]	
+
+			+ SVerticalBox::Slot()
+			.AutoHeight()
+			.HAlign(HAlign_Right)
+			.Padding(2)
+			[
+				SNew(SUniformGridPanel)
+				.SlotPadding(2)
+				+ SUniformGridPanel::Slot(0, 0)
+				[
+					SNew(SButton)
+					.HAlign(HAlign_Center)
+					.Text(LOCTEXT("AnimationModifierWindow_Import", "Apply"))
+					.ToolTipText(LOCTEXT("AnimationModifierWindow_Import_ToolTip", "Apply remove modifiers(s)."))
+					.IsEnabled(this, &SRemoveAnimationModifierContentBrowserWindow::CanApply)
+					.OnClicked(this, &SRemoveAnimationModifierContentBrowserWindow::OnApply)
+				]
+				+ SUniformGridPanel::Slot(1, 0)
+				[
+					SNew(SButton)
+					.HAlign(HAlign_Center)
+					.Text(LOCTEXT("AnimationModifierWindow_Cancel", "Cancel"))
+					.ToolTipText(LOCTEXT("AnimationModifierWindow_Cancel_ToolTip", "Cancels removing modifiers(s)."))
+					.OnClicked(this, &SRemoveAnimationModifierContentBrowserWindow::OnCancel)
+				]
+			]	
+		]
+	];
+}
+
+FReply SRemoveAnimationModifierContentBrowserWindow::OnKeyDown(const FGeometry& MyGeometry, const FKeyEvent& InKeyEvent)
+{
+	if (InKeyEvent.GetKey() == EKeys::Escape)
+	{
+		return OnCancel();
+	}
+
+	return FReply::Unhandled();
+}
+
+void SRemoveAnimationModifierContentBrowserWindow::OnModifierPicked(UClass* PickedClass)
+{
+	UAnimationModifier* Processor = FAnimationModifierHelpers::CreateModifierInstance(GetTransientPackage(), PickedClass);
+
+	Modifiers.Add(Processor);
+
+	FModifierListviewItem* Item = new FModifierListviewItem();
+	Item->Instance = Processor;
+	Item->Class = Processor->GetClass();
+	Item->Index = Modifiers.Num() - 1;
+	Item->OuterClass = nullptr;
+	ModifierItems.Add(ModifierListviewItem(Item));
+
+	// Close the combo box
+	AddModifierCombobox->SetIsOpen(false);
+
+	ModifierListView->Refresh();
+}
+
+void SRemoveAnimationModifierContentBrowserWindow::RemoveModifiersCallback(const TArray<TWeakObjectPtr<UAnimationModifier>>& ModifiersToRemove)
+{
+	Modifiers.RemoveAll([&ModifiersToRemove](UAnimationModifier* Modifier) { return ModifiersToRemove.Contains(Modifier); });
+	ModifierItems.RemoveAll([&ModifiersToRemove](TSharedPtr<FModifierListviewItem> ModifierItem) { return ModifiersToRemove.Contains(ModifierItem->Instance); });
+	ModifierListView->Refresh();
+}
+
+TSharedRef<SWidget> SRemoveAnimationModifierContentBrowserWindow::GenerateAnimationModifierPicker()
+{
+	/** Constrain(s) a picker's classes to those of currently applied pickers */
+	class FRemoveModifierClassFilter : public IClassViewerFilter
+	{
+	public:
+		virtual bool IsClassAllowed(const FClassViewerInitializationOptions& InInitOptions, const UClass* InClass, TSharedRef< FClassViewerFilterFuncs > InFilterFuncs) override
+		{
+			return InClass->IsChildOf(UAnimationModifier::StaticClass()) && AllowList.Contains(InClass);
+		}
+
+		virtual bool IsUnloadedClassAllowed(const FClassViewerInitializationOptions& InInitOptions, const TSharedRef< const IUnloadedBlueprintData > InClass, TSharedRef< FClassViewerFilterFuncs > InFilterFuncs) override
+		{
+			return false;
+		}
+		
+		TSet<UClass*> AllowList;
+	};
+
+	// Customize class picker
+	FClassViewerInitializationOptions Options;
+	Options.bShowUnloadedBlueprints = true;
+	Options.bShowNoneOption = false;
+	const TSharedRef<FRemoveModifierClassFilter> ClassFilter = MakeShared<FRemoveModifierClassFilter>();
+
+	// Query every applied Animation Modifier(s)
+	for (UAnimSequence* AnimationSequence : AnimSequences)
+	{
+		if (UAnimationModifiersAssetUserData* UserData = AnimationSequence->GetAssetUserData<UAnimationModifiersAssetUserData>())
+		{
+			for (UAnimationModifier* const& Modifier : UserData->GetAnimationModifierInstances())
+			{
+				ClassFilter->AllowList.Add(Modifier->GetClass());
+			}
+		}
+	}
+	Options.ClassFilters.Add(ClassFilter);
+
+	return FModuleManager::LoadModuleChecked<FClassViewerModule>("ClassViewer").CreateClassViewer(Options, FOnClassPicked::CreateRaw(this, &SRemoveAnimationModifierContentBrowserWindow::OnModifierPicked));
+}
+
+FReply SRemoveAnimationModifierContentBrowserWindow::OnApply()
+{
+	// Get user confirmation for action(s)
+	const bool bShouldRevert = FMessageDialog::Open(EAppMsgType::YesNo, LOCTEXT("RemoveAndRevertPopupText", "Should the Animation Modifiers be reverted before removing them?"), LOCTEXT("RemoveAndRevertPopupTitle", "Revert before Removing?")) == EAppReturnType::Yes;
+	const bool bShouldRemove = FMessageDialog::Open(EAppMsgType::YesNo, LOCTEXT("RemoveAnimModifiers", "Are you sure you want to remove the anim modifiers from all the selected animation sequences?"), LOCTEXT("RemoveAnimModifiersPopupTitle", "Are you sure?")) == EAppReturnType::Yes;
+	
+	if (!bShouldRemove)
+	{
+		return FReply::Handled();
+	}
+	
+	FScopedTransaction Transaction(LOCTEXT("RemoveModifiersTransaction", "Removing Animation Modifier(s)"));
+	
+	TArray<UAnimationModifiersAssetUserData*> AssetUserData;
+	
+	// Retrieve asset user data for selected sequences
+	for (UAnimSequence* AnimationSequence : AnimSequences)
+	{
+		if (UAnimationModifiersAssetUserData* UserData = AnimationSequence->GetAssetUserData<UAnimationModifiersAssetUserData>())
+		{
+			AssetUserData.Add(UserData);
+		}
+	}
+
+	// For each added modifier create add a new instance to each of the user data entries, using the one(s) set up in the window as template(s)
+	for (int32 i = 0; i < Modifiers.Num(); ++i)
+	{
+		UAnimationModifier* Modifier = Modifiers[i];
+		checkf(Modifier, TEXT("Invalid selected modifier"));
+
+		// Get instance in use by anim asset
+		UAnimationModifier* const* ModifierInstanceInSequence = AssetUserData[i]->GetAnimationModifierInstances().FindByPredicate([Modifier](const UAnimationModifier* TestModifier)
+		{
+			return Modifier->GetClass() == TestModifier->GetClass();
+		});
+
+		checkf(ModifierInstanceInSequence != nullptr, TEXT("Invalid modifier instance"))
+		
+		// Revert modifiers from anim assets
+		if (bShouldRevert)
+		{
+			for (UAnimSequence* AnimSequence : AnimSequences)
+			{
+				(*ModifierInstanceInSequence)->RevertFromAnimationSequence(AnimSequence);
+			
+				// Revert can not fail, thus we can always mark reverted
+				if ((*ModifierInstanceInSequence)->HasLegacyPreviousAppliedModifierOnSkeleton())
+				{
+					checkf(AnimSequence->GetSkeleton() != nullptr, TEXT("Invalid skeleton for anim sequence"));
+					(*ModifierInstanceInSequence)->RemoveLegacyPreviousAppliedModifierOnSkeleton(AnimSequence->GetSkeleton());
+				}
+			}
+		}
+		
+		// Remove modifiers from user assets
+		for (UAnimationModifiersAssetUserData* UserData : AssetUserData)
+		{
+			UAnimationModifier* const* ModifierInstanceInUserData = UserData->GetAnimationModifierInstances().FindByPredicate([Modifier](const UAnimationModifier* TestModifier)
+			{
+				return Modifier->GetClass() == TestModifier->GetClass();
+			});
+
+			if (ModifierInstanceInUserData)
+			{
+				UserData->Modify();
+				UserData->RemoveAnimationModifierInstance(*ModifierInstanceInUserData);
+			}
+		}
+	}
+
+	if (WidgetWindow.IsValid())
+	{
+		WidgetWindow.Pin()->RequestDestroyWindow();
+	}
+
+	return FReply::Handled();
+}
+
+FReply SRemoveAnimationModifierContentBrowserWindow::OnCancel() const
+{
+	if (WidgetWindow.IsValid())
+	{
+		WidgetWindow.Pin()->RequestDestroyWindow();
+	}
+	return FReply::Handled();
+}
+
+bool SRemoveAnimationModifierContentBrowserWindow::CanApply() const
+{
+	return Modifiers.Num() > 0;
+}
+
+#undef LOCTEXT_NAMESPACE // "AnimationModifierContentBrowserWindow"
diff --git a/Engine/Source/Editor/AnimationModifiers/Private/SAnimationModifierContentBrowserWindow.h b/Engine/Source/Editor/AnimationModifiers/Private/SAnimationModifierContentBrowserWindow.h
--- a/Engine/Source/Editor/AnimationModifiers/Private/SAnimationModifierContentBrowserWindow.h
+++ b/Engine/Source/Editor/AnimationModifiers/Private/SAnimationModifierContentBrowserWindow.h
@@ -59,6 +59,59 @@ private:
 	TSharedPtr<SMenuAnchor> AddModifierCombobox;
 	TSharedPtr<SModifierListView> ModifierListView;
 
+	/** Data structures used by the Modifier List View widget */
+	TArray<TSharedPtr<FModifierListviewItem>> ModifierItems;	
+	/** Current set of Animation Modifiers that would be added during Apply */
+	TArray<UAnimationModifier*> Modifiers;
+	/** Previously user-selected Animation Sequences */
+	TArray<UAnimSequence*> AnimSequences;
+};
+
+class SAnimSequencesToBeModifiedViewer;
+
+/** UI slate widget allowing the user to remove Animation Modifier(s) to a selection of Animation Sequences */
+class SRemoveAnimationModifierContentBrowserWindow : public SCompoundWidget
+{
+public:
+	SLATE_BEGIN_ARGS(SRemoveAnimationModifierContentBrowserWindow)
+		: _WidgetWindow()		
+	{}			
+	SLATE_ARGUMENT(TSharedPtr<SWindow>, WidgetWindow)
+	SLATE_ARGUMENT(TArray<UAnimSequence*>, AnimSequences)
+SLATE_END_ARGS()
+
+public:
+	SRemoveAnimationModifierContentBrowserWindow() {}
+	void Construct(const FArguments& InArgs);
+	virtual bool SupportsKeyboardFocus() const override { return true; }
+	virtual FReply OnKeyDown(const FGeometry& MyGeometry, const FKeyEvent& InKeyEvent) override;
+protected:
+
+	/** Callback for when the user picks a specific animation modifier class */
+	void OnModifierPicked(UClass* PickedClass);
+	
+	/** Callback for when the user wants to remove modifier(s) from the listview */
+	void RemoveModifiersCallback(const TArray<TWeakObjectPtr<UAnimationModifier>>& ModifiersToRemove);
+
+	/** Creates a picker that only displays Animation Modifier(s) currently applied to the selected Animation Sequences  */
+	TSharedRef<SWidget> GenerateAnimationModifierPicker();
+	
+	/** Button callback, this applies all currently set up Animation Modifiers to the previously selected Animation Sequences */
+	FReply OnApply();
+	/** Button callback, closes the dialog/window */
+	FReply OnCancel() const;
+	/** Check to see whether or not the user can apply the modifiers in the current state */
+	bool CanApply() const;
+	
+private:
+	/** Window owning this window */
+	TWeakPtr<SWindow> WidgetWindow;
+	
+	TSharedPtr<IDetailsView> ModifierInstanceDetailsView;
+	TSharedPtr<SMenuAnchor> AddModifierCombobox;
+	TSharedPtr<SModifierListView> ModifierListView;
+	TSharedPtr<SAnimSequencesToBeModifiedViewer> AnimSequencesToBeModifiedListViewer;
+	
 	/** Data structures used by the Modifier List View widget */
 	TArray<TSharedPtr<FModifierListviewItem>> ModifierItems;	
 	/** Current set of Animation Modifiers that would be added during Apply */
diff --git a/Engine/Source/Editor/AnimationModifiers/Private/SModifierListview.cpp b/Engine/Source/Editor/AnimationModifiers/Private/SModifierListview.cpp
--- a/Engine/Source/Editor/AnimationModifiers/Private/SModifierListview.cpp
+++ b/Engine/Source/Editor/AnimationModifiers/Private/SModifierListview.cpp
@@ -43,6 +43,7 @@ void SModifierListView::Construct(const FArguments& InArgs)
 	OnRemoveModifierDelegate = InArgs._OnRemoveModifier;
 	OnMoveUpModifierDelegate = InArgs._OnMoveUpModifier;
 	OnMoveDownModifierDelegate = InArgs._OnMoveDownModifier;
+	OnSelectedModifierChangedDelegate = InArgs._OnSelectedModifierChanged;
 	
 	this->ChildSlot
 	[
@@ -58,15 +59,22 @@ TSharedRef<ITableRow> SModifierListView::OnGenerateWidgetForList(ModifierListvie
 void SModifierListView::OnSelectionChanged(ModifierListviewItem SelectedItem, ESelectInfo::Type SelectInfo)
 {
 	UObject* SelectedObject = nullptr;
-
+	TWeakObjectPtr<UAnimationModifier> SelectedModifierInstance = nullptr;
+	
 	// Find object to select if an item was selected (otherwise we are de-selecting to nullptr)
 	if (SelectedItem.IsValid())
 	{
 		SelectedObject = SelectedItem->Instance.Get();
+		SelectedModifierInstance = SelectedItem->Instance;
 	}
 
 	// Set the details view to the currently selected modifier blueprint instance
-	InstanceDetailsView->SetObject(SelectedObject);
+	if (InstanceDetailsView.IsValid())
+	{
+		InstanceDetailsView->SetObject(SelectedObject);
+	}
+	
+	OnSelectedModifierChangedDelegate.ExecuteIfBound(SelectedModifierInstance);
 }
 
 void SModifierListView::OnApplyModifier()
diff --git a/Engine/Source/Editor/AnimationModifiers/Private/SModifierListview.h b/Engine/Source/Editor/AnimationModifiers/Private/SModifierListview.h
--- a/Engine/Source/Editor/AnimationModifiers/Private/SModifierListview.h
+++ b/Engine/Source/Editor/AnimationModifiers/Private/SModifierListview.h
@@ -57,6 +57,7 @@ class SModifierListView : public SCompoundWidget
 	SLATE_ARGUMENT(FOnSingleModifier, OnOpenModifier);
 	SLATE_ARGUMENT(FOnSingleModifier, OnMoveUpModifier);
 	SLATE_ARGUMENT(FOnSingleModifier, OnMoveDownModifier);
+	SLATE_EVENT(FOnSingleModifier, OnSelectedModifierChanged);
 	SLATE_END_ARGS()
 
 	void Construct(const FArguments& InArgs);
@@ -96,6 +97,7 @@ protected:
 	FOnSingleModifier OnOpenModifierDelegate;
 	FOnSingleModifier OnMoveUpModifierDelegate;
 	FOnSingleModifier OnMoveDownModifierDelegate;
+	FOnSingleModifier OnSelectedModifierChangedDelegate;
 	
 	/** Check whether or not selected modifier can moved in either direction */
 	bool CanMoveSelectedItemUp();
diff --git a/Engine/Source/Editor/AnimationModifiers/Public/AnimationModifiersAssetUserData.h b/Engine/Source/Editor/AnimationModifiers/Public/AnimationModifiersAssetUserData.h
--- a/Engine/Source/Editor/AnimationModifiers/Public/AnimationModifiersAssetUserData.h
+++ b/Engine/Source/Editor/AnimationModifiers/Public/AnimationModifiersAssetUserData.h
@@ -24,6 +24,7 @@ class ANIMATIONMODIFIERS_API UAnimationModifiersAssetUserData : public UAssetUse
 
 	friend class SAnimationModifiersTab;
 	friend class SAnimationModifierContentBrowserWindow;
+	friend class SRemoveAnimationModifierContentBrowserWindow;
 	friend class FAnimationModifiersModule;
 	friend class UAnimationModifier;
 public:
diff --git a/Engine/Source/Editor/AnimationModifiers/Public/IAnimationModifiersModule.h b/Engine/Source/Editor/AnimationModifiers/Public/IAnimationModifiersModule.h
--- a/Engine/Source/Editor/AnimationModifiers/Public/IAnimationModifiersModule.h
+++ b/Engine/Source/Editor/AnimationModifiers/Public/IAnimationModifiersModule.h
@@ -13,6 +13,9 @@ public:
 	/** Shows a new modal dialog allowing the user to setup Animation Modifiers to be added for all AnimSequences part of InSequences */
 	virtual void ShowAddAnimationModifierWindow(const TArray<UAnimSequence*>& InSequences) = 0;
 
+	/** Shows a new modal dialog allowing the user to setup Animation Modifiers to be removed for all AnimSequences part of InSequences */
+	virtual void ShowRemoveAnimationModifierWindow(const TArray<UAnimSequence*>& InSequences) = 0;
+
 	/** Apply animation modifiers contained by provided Animation Sequences to its respective outers */
 	virtual void ApplyAnimationModifiers(const TArray<UAnimSequence*>& InSequences, bool bForceApply = true) = 0;
 };
