Fix loctext errors in new "Remove Modifier" feature

diff --git a/Engine/Source/Editor/AnimationModifiers/Private/AnimationModifiersModule.cpp b/Engine/Source/Editor/AnimationModifiers/Private/AnimationModifiersModule.cpp
--- a/Engine/Source/Editor/AnimationModifiers/Private/AnimationModifiersModule.cpp
+++ b/Engine/Source/Editor/AnimationModifiers/Private/AnimationModifiersModule.cpp
@@ -341,7 +341,7 @@ void FAnimationModifiersModule::ShowRemoveAnimationModifierWindow(const TArray<U
 	TSharedPtr<SRemoveAnimationModifierContentBrowserWindow> WindowContent;
 
 	TSharedRef<SWindow> Window = SNew(SWindow)
-		.Title(LOCTEXT("WindowTitle", "Remove Animation Modifier(s)"))
+		.Title(LOCTEXT("RemoveModifiersWindowTitle", "Remove Animation Modifier(s)"))
 		.SizingRule(ESizingRule::UserSized)
 		.ClientSize(FVector2D(500, 500));
 
diff --git a/Engine/Source/Editor/AnimationModifiers/Private/SAnimationModifierContentBrowserWindow.cpp b/Engine/Source/Editor/AnimationModifiers/Private/SAnimationModifierContentBrowserWindow.cpp
--- a/Engine/Source/Editor/AnimationModifiers/Private/SAnimationModifierContentBrowserWindow.cpp
+++ b/Engine/Source/Editor/AnimationModifiers/Private/SAnimationModifierContentBrowserWindow.cpp
@@ -506,8 +506,8 @@ void SRemoveAnimationModifierContentBrowserWindow::Construct(const FArguments& I
 				[
 					SNew(SButton)
 					.HAlign(HAlign_Center)
-					.Text(LOCTEXT("AnimationModifierWindow_Import", "Apply"))
-					.ToolTipText(LOCTEXT("AnimationModifierWindow_Import_ToolTip", "Apply remove modifiers(s)."))
+					.Text(LOCTEXT("AnimationModifierWindow_ApplyRemove", "Apply"))
+					.ToolTipText(LOCTEXT("AnimationModifierWindow_ApplyRemove_ToolTip", "Apply remove modifiers(s)."))
 					.IsEnabled(this, &SRemoveAnimationModifierContentBrowserWindow::CanApply)
 					.OnClicked(this, &SRemoveAnimationModifierContentBrowserWindow::OnApply)
 				]
@@ -515,8 +515,8 @@ void SRemoveAnimationModifierContentBrowserWindow::Construct(const FArguments& I
 				[
 					SNew(SButton)
 					.HAlign(HAlign_Center)
-					.Text(LOCTEXT("AnimationModifierWindow_Cancel", "Cancel"))
-					.ToolTipText(LOCTEXT("AnimationModifierWindow_Cancel_ToolTip", "Cancels removing modifiers(s)."))
+					.Text(LOCTEXT("AnimationModifierWindow_CancelRemove", "Cancel"))
+					.ToolTipText(LOCTEXT("AnimationModifierWindow_CancelRemove_ToolTip", "Cancels removing modifiers(s)."))
 					.OnClicked(this, &SRemoveAnimationModifierContentBrowserWindow::OnCancel)
 				]
 			]	
