Fixed crash when batch removing multiple animation modifiers from a selection of anim sequences and having a animation modifiers only affect a subset of all the selected anim sequences.

#rnx

diff --git a/Engine/Source/Editor/AnimationModifiers/Private/SAnimationModifierContentBrowserWindow.cpp b/Engine/Source/Editor/AnimationModifiers/Private/SAnimationModifierContentBrowserWindow.cpp
--- a/Engine/Source/Editor/AnimationModifiers/Private/SAnimationModifierContentBrowserWindow.cpp
+++ b/Engine/Source/Editor/AnimationModifiers/Private/SAnimationModifierContentBrowserWindow.cpp
@@ -631,26 +631,26 @@ FReply SRemoveAnimationModifierContentBrowserWindow::OnApply()
 		UAnimationModifier* Modifier = Modifiers[i];
 		checkf(Modifier, TEXT("Invalid selected modifier"));
 
-		// Get instance in use by anim asset
-		UAnimationModifier* const* ModifierInstanceInSequence = AssetUserData[i]->GetAnimationModifierInstances().FindByPredicate([Modifier](const UAnimationModifier* TestModifier)
-		{
-			return Modifier->GetClass() == TestModifier->GetClass();
-		});
-
-		checkf(ModifierInstanceInSequence != nullptr, TEXT("Invalid modifier instance"))
-		
 		// Revert modifiers from anim assets
 		if (bShouldRevert)
 		{
 			for (UAnimSequence* AnimSequence : AnimSequences)
 			{
-				(*ModifierInstanceInSequence)->RevertFromAnimationSequence(AnimSequence);
-			
-				// Revert can not fail, thus we can always mark reverted
-				if ((*ModifierInstanceInSequence)->HasLegacyPreviousAppliedModifierOnSkeleton())
+				UAnimationModifier* const* ModifierInstanceInSequence = AssetUserData[i]->GetAnimationModifierInstances().FindByPredicate([Modifier](const UAnimationModifier* TestModifier)
 				{
-					checkf(AnimSequence->GetSkeleton() != nullptr, TEXT("Invalid skeleton for anim sequence"));
-					(*ModifierInstanceInSequence)->RemoveLegacyPreviousAppliedModifierOnSkeleton(AnimSequence->GetSkeleton());
+					return Modifier->GetClass() == TestModifier->GetClass();
+				});
+
+				if (ModifierInstanceInSequence)
+				{
+					(*ModifierInstanceInSequence)->RevertFromAnimationSequence(AnimSequence);
+			
+					// Revert can not fail, thus we can always mark reverted
+					if ((*ModifierInstanceInSequence)->HasLegacyPreviousAppliedModifierOnSkeleton())
+					{
+						checkf(AnimSequence->GetSkeleton() != nullptr, TEXT("Invalid skeleton for anim sequence"));
+						(*ModifierInstanceInSequence)->RemoveLegacyPreviousAppliedModifierOnSkeleton(AnimSequence->GetSkeleton());
+					}
 				}
 			}
 		}
