PR #11891: Expose AddAnimationModifier for Editor Tooling

diff --git a/Engine/Source/Editor/AnimationModifiers/Private/AnimationModifiersAssetUserData.cpp b/Engine/Source/Editor/AnimationModifiers/Private/AnimationModifiersAssetUserData.cpp
--- a/Engine/Source/Editor/AnimationModifiers/Private/AnimationModifiersAssetUserData.cpp
+++ b/Engine/Source/Editor/AnimationModifiers/Private/AnimationModifiersAssetUserData.cpp
@@ -2,6 +2,7 @@
 
 #include "AnimationModifiersAssetUserData.h"
 #include "AnimationModifier.h"
+#include "AnimationModifierHelpers.h"
 #include "UObject/UObjectThreadContext.h"
 
 #include "HAL/PlatformCrt.h"
@@ -35,6 +36,44 @@ const TArray<UAnimationModifier*>& UAnimationModifiersAssetUserData::GetAnimatio
 	return AnimationModifierInstances;
 }
 
+bool UAnimationModifiersAssetUserData::AddAnimationModifierOfClass(UAnimSequenceBase* AnimationSequenceBase, TSubclassOf<UAnimationModifier> InModifierClass)
+{
+	if (AnimationSequenceBase != nullptr)
+	{
+		UClass* ModifierClassToAdd = InModifierClass.Get();
+		if (ModifierClassToAdd != nullptr)
+		{
+			if (UAnimationModifiersAssetUserData* AssetUserData = FAnimationModifierHelpers::RetrieveOrCreateModifierUserData(AnimationSequenceBase))
+			{
+				if(UAnimationModifier* ModifierInstance = FAnimationModifierHelpers::CreateModifierInstance(AssetUserData, ModifierClassToAdd))
+				{					
+					AssetUserData->AddAnimationModifier(ModifierInstance);
+					return true;
+				}
+				else
+				{
+					UE_LOG(LogAnimation, Warning, TEXT("Failed to instantiate Animation Modifier class supplied for AddAnimationModifierOfClass"));
+				}
+			}
+			else
+			{
+				UE_LOG(LogAnimation, Warning, TEXT("Invalid Animation Modifier class supplied for AddAnimationModifierOfClass"));
+			}
+			
+		}
+		else
+		{
+			UE_LOG(LogAnimation, Warning, TEXT("Invalid Animation Modifier class supplied for AddAnimationModifierOfClass"));
+		}
+	}
+	else
+	{
+		UE_LOG(LogAnimation, Warning, TEXT("Invalid Animation Sequence supplied for AddAnimationModifierOfClass"));
+	}
+
+	return false;
+}
+
 void UAnimationModifiersAssetUserData::ChangeAnimationModifierIndex(UAnimationModifier* Instance, int32 Direction)
 {
 	checkf(AnimationModifierInstances.Contains(Instance), TEXT("Instance suppose to be moved is not found"));
diff --git a/Engine/Source/Editor/AnimationModifiers/Public/AnimationModifiersAssetUserData.h b/Engine/Source/Editor/AnimationModifiers/Public/AnimationModifiersAssetUserData.h
--- a/Engine/Source/Editor/AnimationModifiers/Public/AnimationModifiersAssetUserData.h
+++ b/Engine/Source/Editor/AnimationModifiers/Public/AnimationModifiersAssetUserData.h
@@ -30,6 +30,9 @@ class ANIMATIONMODIFIERS_API UAnimationModifiersAssetUserData : public UAssetUse
 public:
 	UAnimationModifiersAssetUserData(const FObjectInitializer& ObjectInitializer);
 	const TArray<UAnimationModifier*>& GetAnimationModifierInstances() const;
+	
+	UFUNCTION(BlueprintCallable, Category=Animation, meta = (ScriptMethod))
+	static bool AddAnimationModifierOfClass(UAnimSequenceBase* AnimationSequenceBase, TSubclassOf<UAnimationModifier> InModifierClass);
 protected:	 
 	/** Begin UAssetUserData overrides */
 	virtual void PostLoad() override;
