CIS non-unity fix

diff --git a/Engine/Source/Editor/AnimationModifiers/Private/AnimationModifiersAssetUserData.cpp b/Engine/Source/Editor/AnimationModifiers/Private/AnimationModifiersAssetUserData.cpp
--- a/Engine/Source/Editor/AnimationModifiers/Private/AnimationModifiersAssetUserData.cpp
+++ b/Engine/Source/Editor/AnimationModifiers/Private/AnimationModifiersAssetUserData.cpp
@@ -1,7 +1,6 @@
 // Copyright Epic Games, Inc. All Rights Reserved.
 
 #include "AnimationModifiersAssetUserData.h"
-#include "AnimationModifier.h"
 #include "AnimationModifierHelpers.h"
 #include "UObject/UObjectThreadContext.h"
 
diff --git a/Engine/Source/Editor/AnimationModifiers/Public/AnimationModifiersAssetUserData.h b/Engine/Source/Editor/AnimationModifiers/Public/AnimationModifiersAssetUserData.h
--- a/Engine/Source/Editor/AnimationModifiers/Public/AnimationModifiersAssetUserData.h
+++ b/Engine/Source/Editor/AnimationModifiers/Public/AnimationModifiersAssetUserData.h
@@ -8,11 +8,11 @@
 #include "UObject/ObjectMacros.h"
 #include "UObject/ObjectPtr.h"
 #include "UObject/UObjectGlobals.h"
+#include "AnimationModifier.h"
 
 #include "AnimationModifiersAssetUserData.generated.h"
 
 class FArchive;
-class UAnimationModifier;
 class UObject;
 template <typename T> struct TObjectPtr;
 
