AnimationModifierAssetUserData::PostEditChangeOwner can cause excessive work when processing interactive changes
- Added FPropertyChangedEvent parameter to UAssetUserData::PostEditChangeOwner, deprecating old signature
- Updated all uses and implementation of PostEditChangeOwner to match new signature
- Use FPropertyChangedEvent to skip interactive changes in AnimationModifiers asset user data

diff --git a/Engine/Plugins/Media/MediaPlate/Source/MediaPlate/Private/MediaPlateAssetUserData.cpp b/Engine/Plugins/Media/MediaPlate/Source/MediaPlate/Private/MediaPlateAssetUserData.cpp
--- a/Engine/Plugins/Media/MediaPlate/Source/MediaPlate/Private/MediaPlateAssetUserData.cpp
+++ b/Engine/Plugins/Media/MediaPlate/Source/MediaPlate/Private/MediaPlateAssetUserData.cpp
@@ -4,7 +4,7 @@
 
 #include UE_INLINE_GENERATED_CPP_BY_NAME(MediaPlateAssetUserData)
 
-void UMediaPlateAssetUserData::PostEditChangeOwner()
+void UMediaPlateAssetUserData::PostEditChangeOwner(const FPropertyChangedEvent& PropertyChangedEvent)
 {
 	OnPostEditChangeOwner.ExecuteIfBound();
 }
diff --git a/Engine/Plugins/Media/MediaPlate/Source/MediaPlate/Public/MediaPlateAssetUserData.h b/Engine/Plugins/Media/MediaPlate/Source/MediaPlate/Public/MediaPlateAssetUserData.h
--- a/Engine/Plugins/Media/MediaPlate/Source/MediaPlate/Public/MediaPlateAssetUserData.h
+++ b/Engine/Plugins/Media/MediaPlate/Source/MediaPlate/Public/MediaPlateAssetUserData.h
@@ -20,6 +20,6 @@ public:
 	FOnPostEditChangeOwner OnPostEditChangeOwner;
 	
 	//~ Begin UAssetUserData
-	virtual void PostEditChangeOwner() override;
+	virtual void PostEditChangeOwner(const FPropertyChangedEvent& PropertyChangedEvent) override;
 	//~ End UAssetUserData
 };
diff --git a/Engine/Source/Editor/AnimationModifiers/Private/AnimationModifiersAssetUserData.cpp b/Engine/Source/Editor/AnimationModifiers/Private/AnimationModifiersAssetUserData.cpp
--- a/Engine/Source/Editor/AnimationModifiers/Private/AnimationModifiersAssetUserData.cpp
+++ b/Engine/Source/Editor/AnimationModifiers/Private/AnimationModifiersAssetUserData.cpp
@@ -91,15 +91,21 @@ void UAnimationModifiersAssetUserData::PostEditChangeProperty(struct FPropertyCh
 }
 
 #if WITH_EDITOR
-void UAnimationModifiersAssetUserData::PostEditChangeOwner()
+void UAnimationModifiersAssetUserData::PostEditChangeOwner(const FPropertyChangedEvent& PropertyChangedEvent)
 {
-	Super::PostEditChangeOwner();
+	Super::PostEditChangeOwner(PropertyChangedEvent);
 
 	// We cant call blueprint implemented functions while routing post load
 	if (FUObjectThreadContext::Get().IsRoutingPostLoad)
 	{
 		return;
 	}
+
+	// Skip interactive changes
+	if (PropertyChangedEvent.ChangeType == EPropertyChangeType::Interactive)
+	{
+		return;
+	}
 	
 	// Only react to PostEditChange if not currently applying modifiers
 	if(!UE::Anim::FApplyModifiersScope::IsScopePending())
diff --git a/Engine/Source/Editor/AnimationModifiers/Public/AnimationModifiersAssetUserData.h b/Engine/Source/Editor/AnimationModifiers/Public/AnimationModifiersAssetUserData.h
--- a/Engine/Source/Editor/AnimationModifiers/Public/AnimationModifiersAssetUserData.h
+++ b/Engine/Source/Editor/AnimationModifiers/Public/AnimationModifiersAssetUserData.h
@@ -39,7 +39,7 @@ protected:
 	virtual void PostEditChangeProperty(struct FPropertyChangedEvent& PropertyChangedEvent) override;	
 	virtual void Serialize(FArchive& Ar) override;
 #if WITH_EDITOR
-	virtual void PostEditChangeOwner() override;
+	virtual void PostEditChangeOwner(const FPropertyChangedEvent& PropertyChangedEvent) override;
 #endif // WITH_EDITOR
 	/** End UAssetUserData overrides */
 
diff --git a/Engine/Source/Runtime/Engine/Classes/Engine/AssetUserData.h b/Engine/Source/Runtime/Engine/Classes/Engine/AssetUserData.h
--- a/Engine/Source/Runtime/Engine/Classes/Engine/AssetUserData.h
+++ b/Engine/Source/Runtime/Engine/Classes/Engine/AssetUserData.h
@@ -25,5 +25,9 @@ class UAssetUserData
 	virtual void DrawCanvas(class FCanvas& Canvas, class FSceneView& View) const {}
 	
 	/** Called when the owner object is modified */
+
+	UE_DEPRECATED(5.6, "PostEditChangeOwner has been deprecated, please override signature with FPropertyChangedEvent parameter instead" )
 	virtual void PostEditChangeOwner() {}
+
+	virtual void PostEditChangeOwner(const FPropertyChangedEvent& PropertyChangedEvent) { PRAGMA_DISABLE_DEPRECATION_WARNINGS PostEditChangeOwner(); PRAGMA_ENABLE_DEPRECATION_WARNINGS }
 };
diff --git a/Engine/Source/Runtime/Engine/Private/AnimationAsset.cpp b/Engine/Source/Runtime/Engine/Private/AnimationAsset.cpp
--- a/Engine/Source/Runtime/Engine/Private/AnimationAsset.cpp
+++ b/Engine/Source/Runtime/Engine/Private/AnimationAsset.cpp
@@ -754,7 +754,7 @@ void UAnimationAsset::PostEditChangeProperty(FPropertyChangedEvent& PropertyChan
 	{
 		if (Datum != nullptr)
 		{
-			Datum->PostEditChangeOwner();
+			Datum->PostEditChangeOwner(PropertyChangedEvent);
 		}
 	}
 
diff --git a/Engine/Source/Runtime/Engine/Private/Components/ActorComponent.cpp b/Engine/Source/Runtime/Engine/Private/Components/ActorComponent.cpp
--- a/Engine/Source/Runtime/Engine/Private/Components/ActorComponent.cpp
+++ b/Engine/Source/Runtime/Engine/Private/Components/ActorComponent.cpp
@@ -1371,14 +1371,14 @@ void UActorComponent::ConsolidatedPostEditChange(const FPropertyChangedEvent& Pr
 	{
 		if (Datum != nullptr)
 		{
-			Datum->PostEditChangeOwner();
+			Datum->PostEditChangeOwner(PropertyChangedEvent);
 		}
 	}
 	for (UAssetUserData* Datum : AssetUserDataEditorOnly)
 	{
 		if (Datum != nullptr)
 		{
-			Datum->PostEditChangeOwner();
+			Datum->PostEditChangeOwner(PropertyChangedEvent);
 		}
 	}
 }
diff --git a/Engine/Source/Runtime/Engine/Private/Level.cpp b/Engine/Source/Runtime/Engine/Private/Level.cpp
--- a/Engine/Source/Runtime/Engine/Private/Level.cpp
+++ b/Engine/Source/Runtime/Engine/Private/Level.cpp
@@ -2774,7 +2774,7 @@ void ULevel::PostEditChangeProperty(FPropertyChangedEvent& PropertyChangedEvent)
 	{
 		if (Datum != nullptr)
 		{
-			Datum->PostEditChangeOwner();
+			Datum->PostEditChangeOwner(PropertyChangedEvent);
 		}
 	}
 } 
diff --git a/Engine/Source/Runtime/Engine/Private/Materials/MaterialInterface.cpp b/Engine/Source/Runtime/Engine/Private/Materials/MaterialInterface.cpp
--- a/Engine/Source/Runtime/Engine/Private/Materials/MaterialInterface.cpp
+++ b/Engine/Source/Runtime/Engine/Private/Materials/MaterialInterface.cpp
@@ -772,7 +772,7 @@ void UMaterialInterface::PostEditChangeProperty(FPropertyChangedEvent& PropertyC
 	{
 		if (Datum != nullptr)
 		{
-			Datum->PostEditChangeOwner();
+			Datum->PostEditChangeOwner(PropertyChangedEvent);
 		}
 	}
 
diff --git a/Engine/Source/Runtime/Engine/Private/SkeletalMesh.cpp b/Engine/Source/Runtime/Engine/Private/SkeletalMesh.cpp
--- a/Engine/Source/Runtime/Engine/Private/SkeletalMesh.cpp
+++ b/Engine/Source/Runtime/Engine/Private/SkeletalMesh.cpp
@@ -1373,14 +1373,14 @@ void USkeletalMesh::PostEditChangeProperty(FPropertyChangedEvent& PropertyChange
 	{
 		if (Datum != nullptr)
 		{
-			Datum->PostEditChangeOwner();
+			Datum->PostEditChangeOwner(PropertyChangedEvent);
 		}
 	}
 	for (UAssetUserData* Datum : AssetUserDataEditorOnly)
 	{
 		if (Datum != nullptr)
 		{
-			Datum->PostEditChangeOwner();
+			Datum->PostEditChangeOwner(PropertyChangedEvent);
 		}
 	}
 
diff --git a/Engine/Source/Runtime/Engine/Private/StaticMesh.cpp b/Engine/Source/Runtime/Engine/Private/StaticMesh.cpp
--- a/Engine/Source/Runtime/Engine/Private/StaticMesh.cpp
+++ b/Engine/Source/Runtime/Engine/Private/StaticMesh.cpp
@@ -5117,7 +5117,7 @@ void UStaticMesh::PostEditChangeProperty(FPropertyChangedEvent& PropertyChangedE
 	{
 		if (Datum != nullptr)
 		{
-			Datum->PostEditChangeOwner();
+			Datum->PostEditChangeOwner(PropertyChangedEvent);
 		}
 	}
 
diff --git a/Engine/Source/Runtime/Engine/Private/Texture.cpp b/Engine/Source/Runtime/Engine/Private/Texture.cpp
--- a/Engine/Source/Runtime/Engine/Private/Texture.cpp
+++ b/Engine/Source/Runtime/Engine/Private/Texture.cpp
@@ -1004,7 +1004,7 @@ PRAGMA_ENABLE_DEPRECATION_WARNINGS
 	{
 		if (Datum != nullptr)
 		{
-			Datum->PostEditChangeOwner();
+			Datum->PostEditChangeOwner(PropertyChangedEvent);
 		}
 	}
 }
diff --git a/Engine/Source/Runtime/Engine/Private/WorldSettings.cpp b/Engine/Source/Runtime/Engine/Private/WorldSettings.cpp
--- a/Engine/Source/Runtime/Engine/Private/WorldSettings.cpp
+++ b/Engine/Source/Runtime/Engine/Private/WorldSettings.cpp
@@ -868,7 +868,7 @@ void AWorldSettings::PostEditChangeProperty(FPropertyChangedEvent& PropertyChang
 	{
 		if (Datum != nullptr)
 		{
-			Datum->PostEditChangeOwner();
+			Datum->PostEditChangeOwner(PropertyChangedEvent);
 		}
 	}
 
@@ -876,7 +876,7 @@ void AWorldSettings::PostEditChangeProperty(FPropertyChangedEvent& PropertyChang
 	{
 		if (Datum != nullptr)
 		{
-			Datum->PostEditChangeOwner();
+			Datum->PostEditChangeOwner(PropertyChangedEvent);
 		}
 	}
 
