Apply All modifiers for USkeleton crashes due to AnimSequence ptr deference

diff --git a/Engine/Source/Editor/AnimationModifiers/Private/AnimationModifier.cpp b/Engine/Source/Editor/AnimationModifiers/Private/AnimationModifier.cpp
--- a/Engine/Source/Editor/AnimationModifiers/Private/AnimationModifier.cpp
+++ b/Engine/Source/Editor/AnimationModifiers/Private/AnimationModifier.cpp
@@ -492,8 +492,12 @@ void UAnimationModifier::ApplyToAll(TSubclassOf<UAnimationModifier> ModifierSubC
 			USkeleton* Skeleton = Cast<USkeleton>(AssetUserData->GetOuter());
 			ensure(AnimSequence || Skeleton);
 
-			IAnimationDataController& Controller = AnimSequence->GetController();
-			IAnimationDataController::FScopedBracket ScopedBracket(Controller, LOCTEXT("ApplyingAnimationModifier", "Applying Animation Modifiers"));
+			TUniquePtr<IAnimationDataController::FScopedBracket> ScopedBracket;
+			if (AnimSequence)
+			{
+				IAnimationDataController& Controller = AnimSequence->GetController();
+				ScopedBracket = MakeUnique<IAnimationDataController::FScopedBracket>(Controller, LOCTEXT("ApplyingAnimationModifier", "Applying Animation Modifiers"));
+			}
 			
 			for (UAnimationModifier* Modifier : AssetUserData->AnimationModifierInstances)
 			{
