Create Montage is causing a freeze
#fix smaller EvaluationModifier lock scopes when applying AnimModifier (as users can hook into ::PostEditChange)

diff --git a/Engine/Source/Editor/AnimationModifiers/Private/AnimationModifier.cpp b/Engine/Source/Editor/AnimationModifiers/Private/AnimationModifier.cpp
--- a/Engine/Source/Editor/AnimationModifiers/Private/AnimationModifier.cpp
+++ b/Engine/Source/Editor/AnimationModifiers/Private/AnimationModifier.cpp
@@ -62,8 +62,6 @@ void UAnimationModifier::ApplyToAnimationSequence(UAnimSequence* AnimSequence) c
 
 	GLog->AddOutputDevice(&OutputLog);
 
-	IAnimationDataModel::FEvaluationAndModificationLock EvaluationAndModificationLock(*AnimSequence->GetDataModelInterface());
-		
 	bool AppliedModifierOuterWasCreated = false;
 	UObject* AppliedModifierOuter = AnimSequence->GetAssetUserData<UAnimationModifiersAssetUserData>();
 	if (!AppliedModifierOuter)
@@ -83,6 +81,7 @@ void UAnimationModifier::ApplyToAnimationSequence(UAnimSequence* AnimSequence) c
 	AnimationDataTransaction.SaveObject(Skeleton);
 
 	{
+		IAnimationDataModel::FEvaluationAndModificationLock EvaluationAndModificationLock(*AnimSequence->GetDataModelInterface());
 		// Group the OnRevert & OnApply operation into one Bracket to reduce compression request
 		IAnimationDataController::FScopedBracket ScopedBracket(AnimSequence->GetController(), LOCTEXT("ApplyModifierBracket", "Applying Animation Modifier"));
 
@@ -127,10 +126,12 @@ void UAnimationModifier::ApplyToAnimationSequence(UAnimSequence* AnimSequence) c
 			FText::FromString(OutputLog));
 		bShouldRevert = !UE::Anim::FApplyModifiersScope::HandleWarning(this, WarningMessage, MessageTitle);
 	}
-
+	
 	// Revert changes if necessary, otherwise post edit and refresh animation data
 	if (bShouldRevert)
 	{
+		IAnimationDataModel::FEvaluationAndModificationLock EvaluationAndModificationLock(*AnimSequence->GetDataModelInterface());
+		
 		AnimationDataTransaction.BeginOperation();
 		AnimationDataTransaction.Apply();
 		AnimationDataTransaction.EndOperation();
