Expose Animation Modifier instance and asset paths to scripting (as read-only)

diff --git a/Engine/Source/Editor/AnimationModifiers/Public/AnimationModifiersAssetUserData.h b/Engine/Source/Editor/AnimationModifiers/Public/AnimationModifiersAssetUserData.h
--- a/Engine/Source/Editor/AnimationModifiers/Public/AnimationModifiersAssetUserData.h
+++ b/Engine/Source/Editor/AnimationModifiers/Public/AnimationModifiersAssetUserData.h
@@ -51,7 +51,7 @@ protected:
 private:
 	UE_API void RemoveInvalidModifiers();
 protected:
-	UPROPERTY()
+	UPROPERTY(BlueprintReadOnly, Category="Animation Modifiers")
 	TArray<TObjectPtr<UAnimationModifier>> AnimationModifierInstances;
 
 	// Animation modifiers APPLIED on the owning animation sequence
@@ -75,7 +75,7 @@ protected:
 	// The "UAnimationModifier::PreviouslyAppliedModifier_DEPRECATED" data will be migrated here
 	// To support proper revert for them
 	// For more compatibility handling, check UAnimationModifier::PostLoad() and GetAppliedModifier() 
-	UPROPERTY()
+	UPROPERTY(BlueprintReadOnly, Category="Animation Modifiers")
 	TMap<FSoftObjectPath, TObjectPtr<UAnimationModifier>> AppliedModifiers;
 };
 
