Fix for animation modifiers not reverting correctly and crashing in some instances

diff --git a/Engine/Source/Editor/AnimationModifiers/Private/SAnimationModifierContentBrowserWindow.cpp b/Engine/Source/Editor/AnimationModifiers/Private/SAnimationModifierContentBrowserWindow.cpp
--- a/Engine/Source/Editor/AnimationModifiers/Private/SAnimationModifierContentBrowserWindow.cpp
+++ b/Engine/Source/Editor/AnimationModifiers/Private/SAnimationModifierContentBrowserWindow.cpp
@@ -613,63 +613,43 @@ FReply SRemoveAnimationModifierContentBrowserWindow::OnApply()
 	
 	FScopedTransaction Transaction(LOCTEXT("RemoveModifiersTransaction", "Removing Animation Modifier(s)"));
 	
-	TArray<UAnimationModifiersAssetUserData*> AssetUserData;
-	
-	// Retrieve asset user data for selected sequences
 	for (UAnimSequence* AnimationSequence : AnimSequences)
 	{
 		if (UAnimationModifiersAssetUserData* UserData = AnimationSequence->GetAssetUserData<UAnimationModifiersAssetUserData>())
 		{
-			AssetUserData.Add(UserData);
-		}
-	}
-
-	// For each added modifier create add a new instance to each of the user data entries, using the one(s) set up in the window as template(s)
-	for (int32 i = 0; i < Modifiers.Num(); ++i)
-	{
-		UAnimationModifier* Modifier = Modifiers[i];
-		checkf(Modifier, TEXT("Invalid selected modifier"));
-
-		// Revert modifiers from anim assets
-		if (bShouldRevert)
-		{
-			for (UAnimSequence* AnimSequence : AnimSequences)
+			const TArray<UAnimationModifier*>& AnimModifierInstancesInUserData = UserData->GetAnimationModifierInstances();
+			for (const UAnimationModifier* Modifier : Modifiers)
 			{
-				UAnimationModifier* const* ModifierInstanceInSequence = AssetUserData[i]->GetAnimationModifierInstances().FindByPredicate([Modifier](const UAnimationModifier* TestModifier)
+				checkf(Modifier, TEXT("Invalid selected modifier"));
+
+				UAnimationModifier* const* ModifierInstanceInUserData = AnimModifierInstancesInUserData.FindByPredicate([Modifier](const UAnimationModifier* TestModifier)
 				{
 					return Modifier->GetClass() == TestModifier->GetClass();
 				});
 
-				if (ModifierInstanceInSequence)
+				if (ModifierInstanceInUserData == nullptr)
 				{
-					(*ModifierInstanceInSequence)->RevertFromAnimationSequence(AnimSequence);
-			
+					continue;
+				}
+
+				if (bShouldRevert)
+				{
+					(*ModifierInstanceInUserData)->RevertFromAnimationSequence(AnimationSequence);
+
 					// Revert can not fail, thus we can always mark reverted
-					if ((*ModifierInstanceInSequence)->HasLegacyPreviousAppliedModifierOnSkeleton())
+					if ((*ModifierInstanceInUserData)->HasLegacyPreviousAppliedModifierOnSkeleton())
 					{
-						checkf(AnimSequence->GetSkeleton() != nullptr, TEXT("Invalid skeleton for anim sequence"));
-						(*ModifierInstanceInSequence)->RemoveLegacyPreviousAppliedModifierOnSkeleton(AnimSequence->GetSkeleton());
+						checkf(AnimationSequence->GetSkeleton() != nullptr, TEXT("Invalid skeleton for anim sequence"));
+						(*ModifierInstanceInUserData)->RemoveLegacyPreviousAppliedModifierOnSkeleton(AnimationSequence->GetSkeleton());
 					}
 				}
-			}
-		}
-		
-		// Remove modifiers from user assets
-		for (UAnimationModifiersAssetUserData* UserData : AssetUserData)
-		{
-			UAnimationModifier* const* ModifierInstanceInUserData = UserData->GetAnimationModifierInstances().FindByPredicate([Modifier](const UAnimationModifier* TestModifier)
-			{
-				return Modifier->GetClass() == TestModifier->GetClass();
-			});
 
-			if (ModifierInstanceInUserData)
-			{
 				UserData->Modify();
 				UserData->RemoveAnimationModifierInstance(*ModifierInstanceInUserData);
 			}
 		}
 	}
-
+	
 	if (WidgetWindow.IsValid())
 	{
 		WidgetWindow.Pin()->RequestDestroyWindow();
